#Fix Updated the Package Project configuration to use a valid config name.

diff --git a/Engine/Source/Editor/MainFrame/Private/Frame/MainFrameActions.cpp b/Engine/Source/Editor/MainFrame/Private/Frame/MainFrameActions.cpp
--- a/Engine/Source/Editor/MainFrame/Private/Frame/MainFrameActions.cpp
+++ b/Engine/Source/Editor/MainFrame/Private/Frame/MainFrameActions.cpp
@@ -597,6 +597,9 @@ void FMainFrameActionCallbacks::PackageProject( const FString InPlatformName, co
 		OptionalParams += TEXT(" -UseDebugParamForEditorExe");
 	}
 
+	FString Configuration = FindObject<UEnum>(ANY_PACKAGE, TEXT("EProjectPackagingBuildConfigurations"))->GetEnumName(PackagingSettings->BuildConfiguration);
+	Configuration = Configuration.Replace(TEXT("PPBC_"), TEXT(""));
+
 	FString ProjectPath = FPaths::IsProjectFilePathSet() ? FPaths::ConvertRelativePathToFull(FPaths::GetProjectFilePath()) : FPaths::RootDir() / FApp::GetGameName() / FApp::GetGameName() + TEXT(".uproject");
 	FString CommandLine = FString::Printf(TEXT("BuildCookRun %s%s -nop4 -project=\"%s\" -cook -allmaps -stage -archive -archivedirectory=\"%s\" -package -%s -clientconfig=%s -ue4exe=%s %s -utf8output"),
 		FRocketSupport::IsRocket() ? TEXT( "-rocket -nocompile" ) : TEXT( "-nocompileeditor" ),
@@ -604,7 +607,7 @@ void FMainFrameActionCallbacks::PackageProject( const FString InPlatformName, co
 		*ProjectPath,
 		*PackagingSettings->StagingDirectory.Path,
 		*PlatformName,
-		*FindObject<UEnum>(ANY_PACKAGE, TEXT("EProjectPackagingBuildConfigurations"))->GetEnumName(PackagingSettings->BuildConfiguration),
+		*Configuration,
 		*ExecutableName,
 		*OptionalParams
 	);
