Translation Editor:
- Added the Launcher localization folder to the list of folders searched for content to be translated
- Also added the localization paths from .ini files exposed through FPaths to this list as well
- Now you can add projects to the Translation Editor through an .ini file like this:

[Internationalization]
+LocalizationPaths=../../../Engine/Content/Localization/Engine

diff --git a/Engine/Source/Editor/MainFrame/Private/Menus/TranslationEditorMenu.h b/Engine/Source/Editor/MainFrame/Private/Menus/TranslationEditorMenu.h
--- a/Engine/Source/Editor/MainFrame/Private/Menus/TranslationEditorMenu.h
+++ b/Engine/Source/Editor/MainFrame/Private/Menus/TranslationEditorMenu.h
@@ -24,57 +24,77 @@ public:
 	 */
 	static void MakeMainFrameTranslationEditorSubMenu( FMenuBuilder& MenuBuilder )
 	{
-		const FString EngineContentLocalizationFolder = FPaths::EngineContentDir() / TEXT("Localization");
-		const FString EngineContentLocalizationFolderWildcard = FString::Printf(TEXT("%s/*"), *EngineContentLocalizationFolder);
-
-		TArray<FString> EngineContentManifestFolders;
-		IFileManager::Get().FindFiles(EngineContentManifestFolders, *EngineContentLocalizationFolderWildcard, false, true);
-
-		TArray<FString> ManifestFileNames;
-		for (FString ProjectManifestFolder : EngineContentManifestFolders)
+		// Directories to search for localization
+		// Directory Structure should look like this:
+		// ProjectRootFolder/Content/Localization/  <-- This is the directory to add to the array below
+		// ProjectRootFolder/Content/Localization/LocalizationProjectName/
+		// ProjectRootFolder/Content/Localization/LocalizationProjectName/LocalizationProjectName.manifest
+		// ProjectRootFolder/Content/Localization/LocalizationProjectName/en/LocalizationProjectName.archive
+		// ProjectRootFolder/Content/Localization/LocalizationProjectName/jp/LocalizationProjectName.archive
+		// ProjectRootFolder/Content/Localization/LocalizationProjectName/ko/LocalizationProjectName.archive
+		// ...
+
+		// Add some specific directories 
+		TArray<FString> DirectoriesToSearch;
+		DirectoriesToSearch.Add(FPaths::EngineContentDir() / TEXT("Localization"));
+		DirectoriesToSearch.Add(FPaths::GameContentDir() / TEXT("Localization"));
+		DirectoriesToSearch.Add(FPaths::EngineDir() / "Programs/NoRedist/UnrealEngineLauncher/Content/Localization");
+
+		// Also check any of these .ini specified localization paths
+		TArray < TArray<FString> > LocalizationPathArrays;
+		LocalizationPathArrays.Add(FPaths::GetGameLocalizationPaths());
+		LocalizationPathArrays.Add(FPaths::GetEditorLocalizationPaths());
+		LocalizationPathArrays.Add(FPaths::GetEngineLocalizationPaths());
+		LocalizationPathArrays.Add(FPaths::GetGameLocalizationPaths());
+		LocalizationPathArrays.Add(FPaths::GetPropertyNameLocalizationPaths());
+		LocalizationPathArrays.Add(FPaths::GetToolTipLocalizationPaths());
+		for (TArray<FString> LocalizationPathArray : LocalizationPathArrays)
 		{
-			// Add root back on
-			ProjectManifestFolder = FPaths::EngineContentDir() / TEXT("Localization") / ProjectManifestFolder;
-
-			const FString ManifestExtension = TEXT("manifest");
-			const FString EngineContentLocalizationFileWildcard = FString::Printf(TEXT("%s/*.%s"), *ProjectManifestFolder, *ManifestExtension);
-			
-			TArray<FString> CurrentFolderManifestFileNames;
-			IFileManager::Get().FindFiles(CurrentFolderManifestFileNames, *EngineContentLocalizationFileWildcard, true, false);
-
-			for (FString& CurrentFolderManifestFileName : CurrentFolderManifestFileNames)
+			for (FString LocalizationPath : LocalizationPathArray)
 			{
-				// Add root back on
-				CurrentFolderManifestFileName = ProjectManifestFolder / CurrentFolderManifestFileName;
+				DirectoriesToSearch.AddUnique(LocalizationPath);
 			}
-
-			ManifestFileNames.Append(CurrentFolderManifestFileNames);
 		}
 
-		// Project Folders
-		const FString ProjectLocalizationFolder = FPaths::GameContentDir() / TEXT("Localization");
-		const FString ProjectLocalizationFolderWildcard = FString::Printf(TEXT("%s/*"), *ProjectLocalizationFolder);
-
-		TArray<FString> ProjectManifestFolders;
-		IFileManager::Get().FindFiles(ProjectManifestFolders, *ProjectLocalizationFolderWildcard, false, true);
-		for (FString ProjectManifestFolder : ProjectManifestFolders)
+		// Find Manifest files
+		TArray<FString> ManifestFileNames;
+		for (FString DirectoryToSearch : DirectoriesToSearch)
 		{
-			// Add root back on
-			ProjectManifestFolder = FPaths::GameContentDir() / TEXT("Localization") / ProjectManifestFolder;
-
+			// Search for .manifest files in this directory
 			const FString ManifestExtension = TEXT("manifest");
-			const FString ProjectLocalizationFileWildcard = FString::Printf(TEXT("%s/*.%s"), *ProjectManifestFolder, *ManifestExtension);
+			const FString TopFolderManifestFileWildcard = FString::Printf(TEXT("%s/*.%s"), *DirectoryToSearch, *ManifestExtension);
 
-			TArray<FString> CurrentFolderManifestFileNames;
-			IFileManager::Get().FindFiles(CurrentFolderManifestFileNames, *ProjectLocalizationFileWildcard, true, false);
+			TArray<FString> TopFolderManifestFileNames;
+			IFileManager::Get().FindFiles(TopFolderManifestFileNames, *TopFolderManifestFileWildcard, true, false);
 
-			for (FString& CurrentFolderManifestFileName : CurrentFolderManifestFileNames)
+			for (FString& CurrentTopFolderManifestFileName : TopFolderManifestFileNames)
 			{
 				// Add root back on
-				CurrentFolderManifestFileName = ProjectManifestFolder / CurrentFolderManifestFileName;
+				CurrentTopFolderManifestFileName = DirectoryToSearch / CurrentTopFolderManifestFileName;
+				ManifestFileNames.AddUnique(CurrentTopFolderManifestFileName);
 			}
 
-			ManifestFileNames.Append(CurrentFolderManifestFileNames);
+			// Search for .manifest files 1 directory deep
+			const FString DirectoryToSearchWildcard = FString::Printf(TEXT("%s/*"), *DirectoryToSearch);
+			TArray<FString> ProjectFolders;
+			IFileManager::Get().FindFiles(ProjectFolders, *DirectoryToSearchWildcard, false, true);
+
+			for (FString ProjectFolder : ProjectFolders)
+			{
+				// Add root back on
+				ProjectFolder = DirectoryToSearch / ProjectFolder;
+				const FString ProjectFolderManifestFileWildcard = FString::Printf(TEXT("%s/*.%s"), *ProjectFolder, *ManifestExtension);
+
+				TArray<FString> CurrentFolderManifestFileNames;
+				IFileManager::Get().FindFiles(CurrentFolderManifestFileNames, *ProjectFolderManifestFileWildcard, true, false);
+
+				for (FString& CurrentFolderManifestFileName : CurrentFolderManifestFileNames)
+				{
+					// Add root back on
+					CurrentFolderManifestFileName = ProjectFolder / CurrentFolderManifestFileName;
+					ManifestFileNames.AddUnique(CurrentFolderManifestFileName);
+				}
+			}
 		}
 
 		MenuBuilder.BeginSection("Project", LOCTEXT("Translation_ChooseProject", "Choose Project:"));
