Translation Editor:
- Remove Duplicated entry in the Directories to Search array

diff --git a/Engine/Source/Editor/MainFrame/Private/Menus/TranslationEditorMenu.h b/Engine/Source/Editor/MainFrame/Private/Menus/TranslationEditorMenu.h
--- a/Engine/Source/Editor/MainFrame/Private/Menus/TranslationEditorMenu.h
+++ b/Engine/Source/Editor/MainFrame/Private/Menus/TranslationEditorMenu.h
@@ -45,7 +45,6 @@ public:
 		LocalizationPathArrays.Add(FPaths::GetGameLocalizationPaths());
 		LocalizationPathArrays.Add(FPaths::GetEditorLocalizationPaths());
 		LocalizationPathArrays.Add(FPaths::GetEngineLocalizationPaths());
-		LocalizationPathArrays.Add(FPaths::GetGameLocalizationPaths());
 		LocalizationPathArrays.Add(FPaths::GetPropertyNameLocalizationPaths());
 		LocalizationPathArrays.Add(FPaths::GetToolTipLocalizationPaths());
 		for (TArray<FString> LocalizationPathArray : LocalizationPathArrays)
