#summary 	Fixed crash on shutdown on Mac
#ttp 330039 	EDITOR: Platform-agnostic editor code depends on Windows-only VSAccessor headers
#detail 	Make sure we dont access modules that are not loaded in ShutdownModule().

diff --git a/Engine/Source/Editor/MainFrame/Private/MainFrameModule.cpp b/Engine/Source/Editor/MainFrame/Private/MainFrameModule.cpp
--- a/Engine/Source/Editor/MainFrame/Private/MainFrameModule.cpp
+++ b/Engine/Source/Editor/MainFrame/Private/MainFrameModule.cpp
@@ -583,10 +583,13 @@ void FMainFrameModule::ShutdownModule( )
 	FModuleManager::Get().OnModuleCompilerFinished().RemoveAll( this );
 
 #if WITH_EDITOR
-	ISourceCodeAccessModule& SourceCodeAccessModule = FModuleManager::LoadModuleChecked<ISourceCodeAccessModule>("SourceCodeAccess");
-	SourceCodeAccessModule.OnLaunchingCodeAccessor().RemoveAll( this );
-	SourceCodeAccessModule.OnDoneLaunchingCodeAccessor().RemoveAll( this );
-	SourceCodeAccessModule.OnOpenFileFailed().RemoveAll( this );
+	if(FModuleManager::Get().IsModuleLoaded("SourceCodeAccess"))
+	{
+		ISourceCodeAccessModule& SourceCodeAccessModule = FModuleManager::GetModuleChecked<ISourceCodeAccessModule>("SourceCodeAccess");
+		SourceCodeAccessModule.OnLaunchingCodeAccessor().RemoveAll( this );
+		SourceCodeAccessModule.OnDoneLaunchingCodeAccessor().RemoveAll( this );
+		SourceCodeAccessModule.OnOpenFileFailed().RemoveAll( this );
+	}
 #endif
 
 	if(CompileStartSound != NULL)
