Enabling the user to specify which build platform they desire, particularly for windows. Win32 and Win64

diff --git a/Engine/Source/Editor/MainFrame/Private/Frame/MainFrameActions.cpp b/Engine/Source/Editor/MainFrame/Private/Frame/MainFrameActions.cpp
--- a/Engine/Source/Editor/MainFrame/Private/Frame/MainFrameActions.cpp
+++ b/Engine/Source/Editor/MainFrame/Private/Frame/MainFrameActions.cpp
@@ -465,7 +465,7 @@ bool FMainFrameActionCallbacks::PackageBuildConfigurationIsChecked( EProjectPack
 	return (GetDefault<UProjectPackagingSettings>()->BuildConfiguration == BuildConfiguration);
 }
 
-void FMainFrameActionCallbacks::PackageProject( const FString InPlatformName, const FText PlatformDisplayName )
+void FMainFrameActionCallbacks::PackageProject( const FString InPlatformName, const FText PlatformDisplayName, const FString AdditionalFlags )
 {
 	FString PlatformName = InPlatformName;
 
@@ -545,14 +545,10 @@ void FMainFrameActionCallbacks::PackageProject( const FString InPlatformName, co
 
 	if (PlatformName == "WindowsNoEditor")
 	{
-		if (PackagingSettings->BuildConfiguration == PPBC_Shipping)
-		{
-			OptionalParams += TEXT(" -targetplatform=Win32");
-		}
-		else
-		{
-			OptionalParams += TEXT(" -targetplatform=Win64");
-		}
+		FString TargetPlatform;
+		FParse::Value(*AdditionalFlags, TEXT("-targetplatform="), TargetPlatform);
+
+		OptionalParams += FString::Printf(TEXT(" -targetplatform=%s"), *TargetPlatform);
 	}
 	else if (PlatformName == "MacNoEditor")
 	{
diff --git a/Engine/Source/Editor/MainFrame/Private/Frame/MainFrameActions.h b/Engine/Source/Editor/MainFrame/Private/Frame/MainFrameActions.h
--- a/Engine/Source/Editor/MainFrame/Private/Frame/MainFrameActions.h
+++ b/Engine/Source/Editor/MainFrame/Private/Frame/MainFrameActions.h
@@ -161,7 +161,7 @@ public:
 	static bool PackageBuildConfigurationIsChecked( EProjectPackagingBuildConfigurations BuildConfiguration );
 
 	/** Packages the project for the specified platform. */
-	static void PackageProject( const FString PlatformName, const FText PlatformDisplayName );
+	static void PackageProject( const FString PlatformName, const FText PlatformDisplayName, const FString AdditionalFlags );
 
 	/** Checks whether a menu action for packaging the project can execute. */
 	static bool PackageProjectCanExecute( const FString PlatformName, bool IsImplemented );
diff --git a/Engine/Source/Editor/MainFrame/Private/Menus/PackageProjectMenu.h b/Engine/Source/Editor/MainFrame/Private/Menus/PackageProjectMenu.h
--- a/Engine/Source/Editor/MainFrame/Private/Menus/PackageProjectMenu.h
+++ b/Engine/Source/Editor/MainFrame/Private/Menus/PackageProjectMenu.h
@@ -41,17 +41,24 @@ public:
 				continue;
 			}
 
-			FUIAction Action(
-				FExecuteAction::CreateStatic(&FMainFrameActionCallbacks::PackageProject, Platform->PlatformName(), Platform->DisplayName()),
-				FCanExecuteAction::CreateStatic(&FMainFrameActionCallbacks::PackageProjectCanExecute, Platform->PlatformName(), Platform->SupportsFeature(ETargetPlatformFeatures::Packaging))
-			);
-
-			MenuBuilder.AddMenuEntry(
-				Platform->DisplayName(),
-				FText::Format(LOCTEXT("PackageGameForPlatformTooltip", "Build, cook and package your game for the {0} platform"), Platform->DisplayName()),
-				FSlateIcon(FEditorStyle::GetStyleSetName(), *FString::Printf(TEXT("Launcher.Platform_%s"), *Platform->PlatformName())),
-				Action
-			);
+			if (Platform->PlatformName() == TEXT("WindowsNoEditor"))
+			{
+				TArray< FText > CollectivePlatforms;
+				CollectivePlatforms.Add(LOCTEXT("Win32","Win32"));
+				CollectivePlatforms.Add(LOCTEXT("Win64","Win64"));
+
+				MenuBuilder.AddSubMenu(
+					Platform->DisplayName(),
+					Platform->DisplayName(),
+					FNewMenuDelegate::CreateStatic(&FPackageProjectMenu::AddPlatformSubPlatformsToMenu, Platform, CollectivePlatforms),
+					false,
+					FSlateIcon(FEditorStyle::GetStyleSetName(), *FString::Printf(TEXT("Launcher.Platform_%s"), *Platform->PlatformName()))
+					);
+			}
+			else
+			{
+				AddPlatformToMenu(MenuBuilder, Platform);
+			}
 		}
 
 		MenuBuilder.AddMenuSeparator();
@@ -62,17 +69,66 @@ public:
 		);
 
 		MenuBuilder.AddMenuSeparator();
-		MenuBuilder.AddMenuEntry( FMainFrameCommands::Get().PackagingSettings );
+		MenuBuilder.AddMenuEntry(FMainFrameCommands::Get().PackagingSettings);
 	}
 
 protected:
 
 	/**
-	 * Creates a build configuration sub-menu.
-	 *
-	 * @param MenuBuilder - The builder for the menu that owns this menu.
-	 */
-	static void MakeBuildConfigurationsMenu( FMenuBuilder& MenuBuilder )
+	* Creates the platform menu entries.
+	*
+	* @param MenuBuilder	- The builder for the menu that owns this menu.
+	* @param Platform		- The target platform we allow packaging for
+	*/
+	static void AddPlatformToMenu(FMenuBuilder& MenuBuilder, const ITargetPlatform* Platform)
+	{
+		FUIAction Action(
+			FExecuteAction::CreateStatic(&FMainFrameActionCallbacks::PackageProject, Platform->PlatformName(), Platform->DisplayName(), FString()),
+			FCanExecuteAction::CreateStatic(&FMainFrameActionCallbacks::PackageProjectCanExecute, Platform->PlatformName(), Platform->SupportsFeature(ETargetPlatformFeatures::Packaging))
+			);
+
+		MenuBuilder.AddMenuEntry(
+			Platform->DisplayName(),
+			FText::Format(LOCTEXT("PackageGameForPlatformTooltip", "Build, cook and package your game for the {0} platform"), Platform->DisplayName()),
+			FSlateIcon(FEditorStyle::GetStyleSetName(), *FString::Printf(TEXT("Launcher.Platform_%s"), *Platform->PlatformName())),
+			Action
+			);
+	}
+
+
+	/**
+	* Creates the platform menu entries for a given platforms sub-platforms.
+	* e.g. Windows has multiple sub-platforms - Win32 and Win64
+	*
+	* @param MenuBuilder	- The builder for the menu that owns this menu.
+	* @param Platform		- The target platform we allow packaging for
+	* @param DisplayNames	- The Sub-platform display names
+	*/
+	static void AddPlatformSubPlatformsToMenu(FMenuBuilder& MenuBuilder, const ITargetPlatform* Platform, TArray<FText> DisplayNames)
+	{
+		for (const FText& DisplayName : DisplayNames)
+		{
+			FString OptionalParams = FString::Printf(TEXT("-targetplatform=%s"), *DisplayName.ToString());
+			FUIAction Action(
+				FExecuteAction::CreateStatic(&FMainFrameActionCallbacks::PackageProject, Platform->PlatformName(), DisplayName, OptionalParams),
+				FCanExecuteAction::CreateStatic(&FMainFrameActionCallbacks::PackageProjectCanExecute, Platform->PlatformName(), Platform->SupportsFeature(ETargetPlatformFeatures::Packaging))
+				);
+
+			MenuBuilder.AddMenuEntry(
+				DisplayName,
+				FText::Format(LOCTEXT("PackageGameForPlatformTooltip", "Build, cook and package your game for the {0} | {1} platform"), Platform->DisplayName(), DisplayName),
+				FSlateIcon(FEditorStyle::GetStyleSetName(), *FString::Printf(TEXT("Launcher.Platform_%s"), *Platform->PlatformName())),
+				Action
+				);
+		}
+	}
+
+	/**
+	* Creates a build configuration sub-menu.
+	*
+	* @param MenuBuilder - The builder for the menu that owns this menu.
+	*/
+	static void MakeBuildConfigurationsMenu(FMenuBuilder& MenuBuilder)
 	{
 		// Only show the debug game option if the game has source code. 
 		TArray<FString> TargetFileNames;
@@ -123,4 +179,4 @@ protected:
 };
 
 
-#undef LOCTEXT_NAMESPACE
+#undef LOCTEXT_NAMESPACE
\ No newline at end of file
