fix for not building the game from Package build without code and not running pre-built binaries
#ue4

diff --git a/Engine/Source/Editor/MainFrame/Private/Frame/MainFrameActions.cpp b/Engine/Source/Editor/MainFrame/Private/Frame/MainFrameActions.cpp
--- a/Engine/Source/Editor/MainFrame/Private/Frame/MainFrameActions.cpp
+++ b/Engine/Source/Editor/MainFrame/Private/Frame/MainFrameActions.cpp
@@ -578,7 +578,7 @@ void FMainFrameActionCallbacks::PackageProject( const FName InPlatformInfoName )
 	}
 
 	// only build if the project has code that might need to be built
-	if (bProjectHasCode && FSourceCodeNavigation::IsCompilerAvailable())
+	if ((bProjectHasCode || !FRocketSupport::IsRocket()) && FSourceCodeNavigation::IsCompilerAvailable())
 	{
 		OptionalParams += TEXT(" -build");
 	}
