Moving the menu entry for credits up a little.

diff --git a/Engine/Source/Editor/MainFrame/Private/Menus/MainMenu.cpp b/Engine/Source/Editor/MainFrame/Private/Menus/MainMenu.cpp
--- a/Engine/Source/Editor/MainFrame/Private/Menus/MainMenu.cpp
+++ b/Engine/Source/Editor/MainFrame/Private/Menus/MainMenu.cpp
@@ -221,6 +221,9 @@ void FMainMenu::FillHelpMenu( FMenuBuilder& MenuBuilder, const TSharedRef< FExte
 
 		MenuBuilder.AddMenuSeparator("EpicGamesHelp");
 		MenuBuilder.AddMenuEntry(FMainFrameCommands::Get().VisitEpicGamesDotCom, "VisitEpicGamesDotCom");
+
+		MenuBuilder.AddMenuSeparator("Credits");
+		MenuBuilder.AddMenuEntry(FMainFrameCommands::Get().CreditsUnrealEd);
 	}
 	MenuBuilder.EndSection();
 
@@ -230,7 +233,6 @@ void FMainMenu::FillHelpMenu( FMenuBuilder& MenuBuilder, const TSharedRef< FExte
 		const FText AboutWindowTitle = NSLOCTEXT("MainHelpMenu", "AboutUnrealEditor", "About Unreal Editor...");
 
 		MenuBuilder.AddMenuEntry(FMainFrameCommands::Get().AboutUnrealEd, "AboutUnrealEd", AboutWindowTitle);
-		MenuBuilder.AddMenuEntry(FMainFrameCommands::Get().CreditsUnrealEd);
 	}
 	MenuBuilder.EndSection();
 #endif
