Disabled the packaging platforms that aren't supported in binary UE4 builds

TTP# 340455 - QA BLOCKER: BINARY: Regression: Code projects fail to package

When using a binary build of UE4 on Windows, DebugGame/Development should always use Win64, and Shipping should always use Win32. I've disabled the relevant menu entries in the Packaging menu so that you're only able to use the valid target.

diff --git a/Engine/Source/Editor/MainFrame/Private/Frame/MainFrameActions.cpp b/Engine/Source/Editor/MainFrame/Private/Frame/MainFrameActions.cpp
--- a/Engine/Source/Editor/MainFrame/Private/Frame/MainFrameActions.cpp
+++ b/Engine/Source/Editor/MainFrame/Private/Frame/MainFrameActions.cpp
@@ -631,6 +631,24 @@ void FMainFrameActionCallbacks::PackageProject( const FName InPlatformInfoName )
 
 bool FMainFrameActionCallbacks::PackageProjectCanExecute( const FName PlatformInfoName, bool IsImplemented )
 {
+	// For a binary Rocket build, Development is ALWAYS Win64, and Shipping is ALWAYS Win32.
+	if(IsImplemented && FRocketSupport::IsRocket())
+	{
+		const EProjectPackagingBuildConfigurations BuildConfiguration = GetDefault<UProjectPackagingSettings>()->BuildConfiguration;
+		switch(BuildConfiguration)
+		{
+		case PPBC_DebugGame: // DebugGame uses the same libs as Development, so fall through and apply the same validation logic
+		case PPBC_Development:
+			return PlatformInfoName != "WindowsNoEditor_Win32";
+
+		case PPBC_Shipping:
+			return PlatformInfoName != "WindowsNoEditor_Win64";
+
+		default:
+			break;
+		}
+	}
+
 	return IsImplemented;
 }
 
