Remove double EOLs from log output during packaging in the editor
#ttp 341969

diff --git a/Engine/Source/Editor/MainFrame/Private/Frame/MainFrameActions.cpp b/Engine/Source/Editor/MainFrame/Private/Frame/MainFrameActions.cpp
--- a/Engine/Source/Editor/MainFrame/Private/Frame/MainFrameActions.cpp
+++ b/Engine/Source/Editor/MainFrame/Private/Frame/MainFrameActions.cpp
@@ -1228,7 +1228,7 @@ void FMainFrameActionCallbacks::HandleUatProcessCompleted( int32 ReturnCode, TWe
 
 void FMainFrameActionCallbacks::HandleUatProcessOutput( FString Output, TWeakPtr<SNotificationItem> NotificationItemPtr, FText PlatformDisplayName, FText TaskName )
 {
-	if (!Output.IsEmpty())
+	if (!Output.IsEmpty() && !Output.Equals("\r"))
 	{
 		UE_LOG(MainFrameActions, Log, TEXT("%s (%s): %s"), *TaskName.ToString(), *PlatformDisplayName.ToString(), *Output);
 	}	
