CrashReport - Added crash reporter to builds built by "Package Project"

diff --git a/Engine/Source/Editor/MainFrame/Private/Frame/MainFrameActions.cpp b/Engine/Source/Editor/MainFrame/Private/Frame/MainFrameActions.cpp
--- a/Engine/Source/Editor/MainFrame/Private/Frame/MainFrameActions.cpp
+++ b/Engine/Source/Editor/MainFrame/Private/Frame/MainFrameActions.cpp
@@ -624,7 +624,7 @@ void FMainFrameActionCallbacks::PackageProject( const FName InPlatformInfoName )
 	Configuration = Configuration.Replace(TEXT("PPBC_"), TEXT(""));
 
 	FString ProjectPath = FPaths::IsProjectFilePathSet() ? FPaths::ConvertRelativePathToFull(FPaths::GetProjectFilePath()) : FPaths::RootDir() / FApp::GetGameName() / FApp::GetGameName() + TEXT(".uproject");
-	FString CommandLine = FString::Printf(TEXT("BuildCookRun %s%s -nop4 -project=\"%s\" -cook -allmaps -stage -archive -archivedirectory=\"%s\" -package -%s -clientconfig=%s -ue4exe=%s %s -utf8output"),
+	FString CommandLine = FString::Printf(TEXT("BuildCookRun %s%s -nop4 -project=\"%s\" -cook -allmaps -CrashReporter -stage -archive -archivedirectory=\"%s\" -package -%s -clientconfig=%s -ue4exe=%s %s -utf8output"),
 		FRocketSupport::IsRocket() ? TEXT( "-rocket -nocompile" ) : TEXT( "-nocompileeditor" ),
 		FApp::IsEngineInstalled() ? TEXT(" -installed") : TEXT(""),
 		*ProjectPath,
