addition of errors to Launch On, Package, and Cook analytics
TTP341794
#ue4
#analytics

diff --git a/Engine/Build/BatchFiles/RunUAT.bat b/Engine/Build/BatchFiles/RunUAT.bat
--- a/Engine/Build/BatchFiles/RunUAT.bat
+++ b/Engine/Build/BatchFiles/RunUAT.bat
@@ -111,7 +111,7 @@ goto Exit_Failure
 :Exit_Failure
 echo BUILD FAILED
 popd
-exit /B 1
+exit /B %ERRORLEVEL%
 
 :Exit
 rem ## Restore original CWD in case we change it
diff --git a/Engine/Build/BatchFiles/RunUAT.sh b/Engine/Build/BatchFiles/RunUAT.sh
--- a/Engine/Build/BatchFiles/RunUAT.sh
+++ b/Engine/Build/BatchFiles/RunUAT.sh
@@ -107,5 +107,5 @@ UATReturn=$?
 
 if [ $UATReturn -ne 0 ]; then
 	echo RunUAT ERROR: AutomationTool was unable to run successfully.
-	exit 1
+	exit $UATReturn
 fi
diff --git a/Engine/Source/Developer/LauncherServices/Private/Launcher/LauncherTask.h b/Engine/Source/Developer/LauncherServices/Private/Launcher/LauncherTask.h
--- a/Engine/Source/Developer/LauncherServices/Private/Launcher/LauncherTask.h
+++ b/Engine/Source/Developer/LauncherServices/Private/Launcher/LauncherTask.h
@@ -24,6 +24,7 @@ public:
 		, Status(ELauncherTaskStatus::Pending)
 		, ReadPipe(InReadPipe)
 		, WritePipe(InWritePipe)
+		, Result(0)
 	{ }
 
 public:
@@ -113,6 +114,22 @@ public:
 		return Status == ELauncherTaskStatus::Completed;
 	}
 
+	int32 ReturnCode() const
+	{
+		if (IsChainFinished())
+		{
+			for (int32 ContinuationIndex = 0; ContinuationIndex < Continuations.Num(); ++ContinuationIndex)
+			{
+				if (Continuations[ContinuationIndex]->ReturnCode() != 0)
+				{
+					return Continuations[ContinuationIndex]->ReturnCode();
+				}
+			}
+		}
+
+		return Result;
+	}
+
 public:
 
 	// Begin FRunnable interface
@@ -305,4 +322,7 @@ protected:
 	// read and write pipe
 	void* ReadPipe;
 	void* WritePipe;
+
+	// result
+	int32 Result;
 };
diff --git a/Engine/Source/Developer/LauncherServices/Private/Launcher/LauncherUATTask.h b/Engine/Source/Developer/LauncherServices/Private/Launcher/LauncherUATTask.h
--- a/Engine/Source/Developer/LauncherServices/Private/Launcher/LauncherUATTask.h
+++ b/Engine/Source/Developer/LauncherServices/Private/Launcher/LauncherUATTask.h
@@ -135,8 +135,7 @@ protected:
 			FPlatformProcess::Sleep(0.25);
 		}
 
-		int32 ReturnCode;
-		if (!FPlatformProcess::GetProcReturnCode(ProcessHandle, &ReturnCode))
+		if (!FPlatformProcess::GetProcReturnCode(ProcessHandle, &Result))
 		{
 			return false;
 		}
@@ -146,7 +145,7 @@ protected:
 			return false;
 		}
 
-		return (ReturnCode == 0);
+		return (Result == 0);
 	}
 
 private:
diff --git a/Engine/Source/Developer/LauncherServices/Private/Launcher/LauncherWorker.cpp b/Engine/Source/Developer/LauncherServices/Private/Launcher/LauncherWorker.cpp
--- a/Engine/Source/Developer/LauncherServices/Private/Launcher/LauncherWorker.cpp
+++ b/Engine/Source/Developer/LauncherServices/Private/Launcher/LauncherWorker.cpp
@@ -107,7 +107,7 @@ uint32 FLauncherWorker::Run( )
 	}
 	else
 	{
-		LaunchCompleted.Broadcast(TaskChain->Succeeded(), FPlatformTime::Seconds() - LaunchStartTime);
+		LaunchCompleted.Broadcast(TaskChain->Succeeded(), FPlatformTime::Seconds() - LaunchStartTime, TaskChain->ReturnCode());
 	}
 
 	FPlatformProcess::ClosePipe(ReadPipe, WritePipe);
diff --git a/Engine/Source/Developer/LauncherServices/Public/Interfaces/ILauncherTask.h b/Engine/Source/Developer/LauncherServices/Public/Interfaces/ILauncherTask.h
--- a/Engine/Source/Developer/LauncherServices/Public/Interfaces/ILauncherTask.h
+++ b/Engine/Source/Developer/LauncherServices/Public/Interfaces/ILauncherTask.h
@@ -97,6 +97,14 @@ public:
 	 */
 	virtual bool IsFinished( ) const = 0;
 
+	/**
+	 * retrieves the return code from the task
+	 *
+	 *
+	 * @return return code from the task
+	 */
+	virtual int32 ReturnCode() const = 0;
+
 	/**
 	 * Gets the stage completed delegate
 	 *
diff --git a/Engine/Source/Developer/LauncherServices/Public/Interfaces/ILauncherWorker.h b/Engine/Source/Developer/LauncherServices/Public/Interfaces/ILauncherWorker.h
--- a/Engine/Source/Developer/LauncherServices/Public/Interfaces/ILauncherWorker.h
+++ b/Engine/Source/Developer/LauncherServices/Public/Interfaces/ILauncherWorker.h
@@ -35,7 +35,7 @@ DECLARE_MULTICAST_DELEGATE_OneParam(FOnStageStartedDelegate, const FString&);
 DECLARE_MULTICAST_DELEGATE_TwoParams(FOnStageCompletedDelegate, const FString&, double);
 
 /** Delegate used to notify when the launch is complete */
-DECLARE_MULTICAST_DELEGATE_TwoParams(FOnLaunchCompletedDelegate, bool, double);
+DECLARE_MULTICAST_DELEGATE_ThreeParams(FOnLaunchCompletedDelegate, bool, double, int32);
 
 /** Delegate used to notify when the launch was canceled */
 DECLARE_MULTICAST_DELEGATE_OneParam(FOnLaunchCanceledDelegate, double);
diff --git a/Engine/Source/Developer/TargetPlatform/Public/Common/TargetPlatformBase.h b/Engine/Source/Developer/TargetPlatform/Public/Common/TargetPlatformBase.h
--- a/Engine/Source/Developer/TargetPlatform/Public/Common/TargetPlatformBase.h
+++ b/Engine/Source/Developer/TargetPlatform/Public/Common/TargetPlatformBase.h
@@ -170,7 +170,7 @@ public:
 		return true;
 	}
 
-	virtual int IsReadyToBuild(const FString& ProjectPath, bool bProjectHasCode, FString& OutDocumentationPath) const override
+	virtual int DoesntHaveRequirements(const FString& ProjectPath, bool bProjectHasCode, FString& OutDocumentationPath) const override
 	{
 		int bReadyToBuild = ETargetPlatformReadyStatus::Ready; // @todo How do we check that the iOS SDK is installed when building from Windows? Is that even possible?
 		if (!IsSdkInstalled(bProjectHasCode, OutDocumentationPath))
diff --git a/Engine/Source/Developer/TargetPlatform/Public/Interfaces/ITargetPlatform.h b/Engine/Source/Developer/TargetPlatform/Public/Interfaces/ITargetPlatform.h
--- a/Engine/Source/Developer/TargetPlatform/Public/Interfaces/ITargetPlatform.h
+++ b/Engine/Source/Developer/TargetPlatform/Public/Interfaces/ITargetPlatform.h
@@ -306,7 +306,7 @@ public:
 	 * @param OutDocumentationPath Let's the platform tell the editor a path to show some information about how to fix any problem
 	 * @return Readiness status
 	 */
-	virtual int IsReadyToBuild(const FString& ProjectPath, bool bProjectHasCode, FString& OutDocumentationPath) const = 0;
+	virtual int DoesntHaveRequirements(const FString& ProjectPath, bool bProjectHasCode, FString& OutDocumentationPath) const = 0;
 
 #if WITH_ENGINE
 	/**
diff --git a/Engine/Source/Developer/iOS/IOSTargetPlatform/Private/IOSTargetPlatform.cpp b/Engine/Source/Developer/iOS/IOSTargetPlatform/Private/IOSTargetPlatform.cpp
--- a/Engine/Source/Developer/iOS/IOSTargetPlatform/Private/IOSTargetPlatform.cpp
+++ b/Engine/Source/Developer/iOS/IOSTargetPlatform/Private/IOSTargetPlatform.cpp
@@ -99,7 +99,7 @@ bool FIOSTargetPlatform::IsSdkInstalled(bool bProjectHasCode, FString& OutDocume
 	return biOSSDKInstalled;
 }
 
-int FIOSTargetPlatform::IsReadyToBuild(const FString& ProjectPath, bool bProjectHasCode, FString& OutDocumentationPath) const
+int FIOSTargetPlatform::DoesntHaveRequirements(const FString& ProjectPath, bool bProjectHasCode, FString& OutDocumentationPath) const
 {
 	int bReadyToBuild = ETargetPlatformReadyStatus::Ready; // @todo How do we check that the iOS SDK is installed when building from Windows? Is that even possible?
 	if (!IsSdkInstalled(bProjectHasCode, OutDocumentationPath))
diff --git a/Engine/Source/Developer/iOS/IOSTargetPlatform/Private/IOSTargetPlatform.h b/Engine/Source/Developer/iOS/IOSTargetPlatform/Private/IOSTargetPlatform.h
--- a/Engine/Source/Developer/iOS/IOSTargetPlatform/Private/IOSTargetPlatform.h
+++ b/Engine/Source/Developer/iOS/IOSTargetPlatform/Private/IOSTargetPlatform.h
@@ -84,7 +84,7 @@ public:
 	}
 
 	virtual bool IsSdkInstalled(bool bProjectHasCode, FString& OutDocumentationPath) const override;
-	virtual int IsReadyToBuild(const FString& ProjectPath, bool bProjectHasCode, FString& OutDocumentationPath) const override;
+	virtual int DoesntHaveRequirements(const FString& ProjectPath, bool bProjectHasCode, FString& OutDocumentationPath) const override;
 
 
 #if WITH_ENGINE
diff --git a/Engine/Source/Editor/MainFrame/Private/Frame/MainFrameActions.cpp b/Engine/Source/Editor/MainFrame/Private/Frame/MainFrameActions.cpp
--- a/Engine/Source/Editor/MainFrame/Private/Frame/MainFrameActions.cpp
+++ b/Engine/Source/Editor/MainFrame/Private/Frame/MainFrameActions.cpp
@@ -3,7 +3,6 @@
 
 #include "MainFramePrivatePCH.h"
 
-
 #define LOCTEXT_NAMESPACE "MainFrameActions"
 
 DEFINE_LOG_CATEGORY_STATIC(MainFrameActions, Log, All);
@@ -485,6 +484,9 @@ void FMainFrameActionCallbacks::PackageProject( const FName InPlatformInfoName )
 	{
 		IMainFrameModule& MainFrameModule = FModuleManager::GetModuleChecked<IMainFrameModule>(TEXT("MainFrame"));
 		MainFrameModule.BroadcastMainFrameSDKNotInstalled(PlatformInfo->TargetPlatformName.ToString(), PlatformInfo->SDKTutorial);
+		TArray<FAnalyticsEventAttribute> ParamArray;
+		ParamArray.Add(FAnalyticsEventAttribute(TEXT("Time"), 0.0));
+		FEditorAnalytics::ReportEvent(TEXT("Editor.Package.Failed"), PlatformInfo->TargetPlatformName.ToString(), bProjectHasCode, EAnalyticsErrorCodes::SDKNotFound, ParamArray);
 		return;
 	}
 
@@ -494,7 +496,9 @@ void FMainFrameActionCallbacks::PackageProject( const FName InPlatformInfoName )
 		{
 			FString NotInstalledDocLink;
 			FString ProjectPath = FPaths::IsProjectFilePathSet() ? FPaths::ConvertRelativePathToFull(FPaths::GetProjectFilePath()) : FPaths::RootDir() / FApp::GetGameName() / FApp::GetGameName() + TEXT(".uproject");
-			switch (Platform->IsReadyToBuild(ProjectPath, bProjectHasCode, NotInstalledDocLink))
+			int32 Result = Platform->DoesntHaveRequirements(ProjectPath, bProjectHasCode, NotInstalledDocLink);
+			FEditorAnalytics::ReportBuildRequirementsFailure(TEXT("Editor.Package.Failed"), PlatformInfo->TargetPlatformName.ToString(), bProjectHasCode, Result);
+			switch (Result)
 			{
 				// broadcast this, and assume someone will pick it up
 			case ETargetPlatformReadyStatus::SDKNotFound:
@@ -1010,12 +1014,19 @@ void FMainFrameActionCallbacks::CreateUatTask( const FString& CommandLine, const
 #endif
 
 	FString UatPath = FPaths::ConvertRelativePathToFull(FPaths::EngineDir() / TEXT("Build/BatchFiles") / RunUATScriptName);
+	FGameProjectGenerationModule& GameProjectModule = FModuleManager::LoadModuleChecked<FGameProjectGenerationModule>(TEXT("GameProjectGeneration"));
+	bool bHasCode = GameProjectModule.Get().GetProjectCodeFileCount() > 0;
 
 	if (!FPaths::FileExists(UatPath))
 	{
 		FFormatNamedArguments Arguments;
 		Arguments.Add(TEXT("File"), FText::FromString(UatPath));
 		FMessageDialog::Open(EAppMsgType::Ok, FText::Format(LOCTEXT("RequiredFileNotFoundMessage", "A required file could not be found:\n{File}"), Arguments));
+
+		TArray<FAnalyticsEventAttribute> ParamArray;
+		ParamArray.Add(FAnalyticsEventAttribute(TEXT("Time"), 0.0));
+		FString EventName = (CommandLine.Contains(TEXT("-package")) ? TEXT("Editor.Package") : TEXT("Editor.Cook"));
+		FEditorAnalytics::ReportEvent(EventName + TEXT(".Failed"), PlatformDisplayName.ToString(), bHasCode, EAnalyticsErrorCodes::UATNotFound, ParamArray);
 		
 		return;
 	}
@@ -1054,20 +1065,8 @@ void FMainFrameActionCallbacks::CreateUatTask( const FString& CommandLine, const
 		return;
 	}
 
-	FGameProjectGenerationModule& GameProjectModule = FModuleManager::LoadModuleChecked<FGameProjectGenerationModule>(TEXT("GameProjectGeneration"));
-	bool bProjectHasCode = GameProjectModule.Get().GetProjectCodeFileCount() > 0;
-
 	FString EventName = (CommandLine.Contains(TEXT("-package")) ? TEXT("Editor.Package") : TEXT("Editor.Cook"));
-	if( FEngineAnalytics::IsAvailable() )
-	{
-		const UGeneralProjectSettings& ProjectSettings = *GetDefault<UGeneralProjectSettings>();
-		TArray<FAnalyticsEventAttribute> ParamArray;
-		ParamArray.Add(FAnalyticsEventAttribute(TEXT("ProjectID"), ProjectSettings.ProjectID.ToString()));
-		ParamArray.Add(FAnalyticsEventAttribute(TEXT("Platform"), PlatformDisplayName.ToString()));
-		ParamArray.Add(FAnalyticsEventAttribute(TEXT("ProjectType"), bProjectHasCode ? TEXT("C++ Code") : TEXT("Content Only")));
-
-		FEngineAnalytics::GetProvider().RecordEvent( EventName + TEXT(".Start"), ParamArray );
-	}
+	FEditorAnalytics::ReportEvent(EventName + TEXT(".Start"), PlatformDisplayName.ToString(), bHasCode);
 
 	NotificationItem->SetCompletionState(SNotificationItem::CS_Pending);
 
@@ -1077,7 +1076,7 @@ void FMainFrameActionCallbacks::CreateUatTask( const FString& CommandLine, const
 	EventData Data;
 	Data.StartTime = FPlatformTime::Seconds();
 	Data.EventName = EventName;
-	Data.bProjectHasCode = bProjectHasCode;
+	Data.bProjectHasCode = bHasCode;
 	UatProcess->OnCanceled().BindStatic(&FMainFrameActionCallbacks::HandleUatProcessCanceled, NotificationItemPtr, PlatformDisplayName, TaskShortName, Data);
 	UatProcess->OnCompleted().BindStatic(&FMainFrameActionCallbacks::HandleUatProcessCompleted, NotificationItemPtr, PlatformDisplayName, TaskShortName, Data);
 	UatProcess->OnOutput().BindStatic(&FMainFrameActionCallbacks::HandleUatProcessOutput, NotificationItemPtr, PlatformDisplayName, TaskShortName);
@@ -1093,10 +1092,13 @@ void FMainFrameActionCallbacks::CreateUatTask( const FString& CommandLine, const
 		NotificationItem->SetText(LOCTEXT("UatLaunchFailedNotification", "Failed to launch Unreal Automation Tool (UAT)!"));
 		NotificationItem->SetCompletionState(SNotificationItem::CS_Fail);
 		NotificationItem->ExpireAndFadeout();
+
+		TArray<FAnalyticsEventAttribute> ParamArray;
+		ParamArray.Add(FAnalyticsEventAttribute(TEXT("Time"), 0.0));
+		FEditorAnalytics::ReportEvent(EventName + TEXT(".Failed"), PlatformDisplayName.ToString(), bHasCode, EAnalyticsErrorCodes::UATLaunchFailure, ParamArray);
 	}
 }
 
-
 /* FMainFrameActionCallbacks callbacks
  *****************************************************************************/
 
@@ -1172,17 +1174,9 @@ void FMainFrameActionCallbacks::HandleUatProcessCanceled( TWeakPtr<SNotification
 		FText::Format(LOCTEXT("UatProcessFailedNotification", "{TaskName} canceled!"), Arguments)
 	);
 
-	if( FEngineAnalytics::IsAvailable() )
-	{
-		const UGeneralProjectSettings& ProjectSettings = *GetDefault<UGeneralProjectSettings>();
-		TArray<FAnalyticsEventAttribute> ParamArray;
-		ParamArray.Add(FAnalyticsEventAttribute(TEXT("ProjectID"), ProjectSettings.ProjectID.ToString()));
-		ParamArray.Add(FAnalyticsEventAttribute(TEXT("Platform"), PlatformDisplayName.ToString()));
-		ParamArray.Add(FAnalyticsEventAttribute(TEXT("ProjectType"), Event.bProjectHasCode ? TEXT("C++ Code") : TEXT("Content Only")));
-		ParamArray.Add(FAnalyticsEventAttribute(TEXT("Time"), FPlatformTime::Seconds() - Event.StartTime));
-
-		FEngineAnalytics::GetProvider().RecordEvent( Event.EventName + TEXT(".Canceled"), ParamArray );
-	}
+	TArray<FAnalyticsEventAttribute> ParamArray;
+	ParamArray.Add(FAnalyticsEventAttribute(TEXT("Time"), FPlatformTime::Seconds() - Event.StartTime));
+	FEditorAnalytics::ReportEvent(Event.EventName + TEXT(".Canceled"), PlatformDisplayName.ToString(), Event.bProjectHasCode, ParamArray);
 //	FMessageLog("PackagingResults").Warning(FText::Format(LOCTEXT("UatProcessCanceledMessageLog", "{TaskName} for {Platform} canceled by user"), Arguments));
 }
 
@@ -1201,17 +1195,9 @@ void FMainFrameActionCallbacks::HandleUatProcessCompleted( int32 ReturnCode, TWe
 			FText::Format(LOCTEXT("UatProcessSucceededNotification", "{TaskName} complete!"), Arguments)
 		);
 
-		if( FEngineAnalytics::IsAvailable() )
-		{
-			const UGeneralProjectSettings& ProjectSettings = *GetDefault<UGeneralProjectSettings>();
-			TArray<FAnalyticsEventAttribute> ParamArray;
-			ParamArray.Add(FAnalyticsEventAttribute(TEXT("ProjectID"), ProjectSettings.ProjectID.ToString()));
-			ParamArray.Add(FAnalyticsEventAttribute(TEXT("Platform"), PlatformDisplayName.ToString()));
-			ParamArray.Add(FAnalyticsEventAttribute(TEXT("ProjectType"), Event.bProjectHasCode ? TEXT("C++ Code") : TEXT("Content Only")));
-			ParamArray.Add(FAnalyticsEventAttribute(TEXT("Time"), FPlatformTime::Seconds() - Event.StartTime));
-
-			FEngineAnalytics::GetProvider().RecordEvent( Event.EventName + TEXT(".Completed"), ParamArray );
-		}
+		TArray<FAnalyticsEventAttribute> ParamArray;
+		ParamArray.Add(FAnalyticsEventAttribute(TEXT("Time"), FPlatformTime::Seconds() - Event.StartTime));
+		FEditorAnalytics::ReportEvent(Event.EventName + TEXT(".Completed"), PlatformDisplayName.ToString(), Event.bProjectHasCode, ParamArray);
 
 //		FMessageLog("PackagingResults").Info(FText::Format(LOCTEXT("UatProcessSuccessMessageLog", "{TaskName} for {Platform} completed successfully"), Arguments));
 	}
@@ -1223,17 +1209,9 @@ void FMainFrameActionCallbacks::HandleUatProcessCompleted( int32 ReturnCode, TWe
 			FText::Format(LOCTEXT("PackagerFailedNotification", "{TaskName} failed!"), Arguments)
 		);
 
-		if( FEngineAnalytics::IsAvailable() )
-		{
-			const UGeneralProjectSettings& ProjectSettings = *GetDefault<UGeneralProjectSettings>();
-			TArray<FAnalyticsEventAttribute> ParamArray;
-			ParamArray.Add(FAnalyticsEventAttribute(TEXT("ProjectID"), ProjectSettings.ProjectID.ToString()));
-			ParamArray.Add(FAnalyticsEventAttribute(TEXT("Platform"), PlatformDisplayName.ToString()));
-			ParamArray.Add(FAnalyticsEventAttribute(TEXT("ProjectType"), Event.bProjectHasCode ? TEXT("C++ Code") : TEXT("Content Only")));
-			ParamArray.Add(FAnalyticsEventAttribute(TEXT("Time"), FPlatformTime::Seconds() - Event.StartTime));
-
-			FEngineAnalytics::GetProvider().RecordEvent( Event.EventName + TEXT(".Completed"), ParamArray );
-		}
+		TArray<FAnalyticsEventAttribute> ParamArray;
+		ParamArray.Add(FAnalyticsEventAttribute(TEXT("Time"), FPlatformTime::Seconds() - Event.StartTime));
+		FEditorAnalytics::ReportEvent(Event.EventName + TEXT(".Failed"), PlatformDisplayName.ToString(), Event.bProjectHasCode, ReturnCode, ParamArray);
 
 //		FMessageLog("PackagingResults").Info(FText::Format(LOCTEXT("UatProcessFailedMessageLog", "{TaskName} for {Platform} failed"), Arguments));
 	}
diff --git a/Engine/Source/Editor/MainFrame/Private/Frame/MainFrameActions.h b/Engine/Source/Editor/MainFrame/Private/Frame/MainFrameActions.h
--- a/Engine/Source/Editor/MainFrame/Private/Frame/MainFrameActions.h
+++ b/Engine/Source/Editor/MainFrame/Private/Frame/MainFrameActions.h
@@ -239,6 +239,15 @@ protected:
 	 */
 	static void CreateUatTask( const FString& CommandLine, const FText& PlatformDisplayName, const FText& TaskName, const FText &TaskShortName, const FSlateBrush* TaskIcon );
 
+	/**
+	 * Translate the given error code to human readable text
+	 *
+	 * @param ErrorCode - the error code returned from UAT task
+	 *
+	 * @return the human readable error string
+	 */
+	static FString TranslateUATError(int32 ErrorCode);
+
 	/** Callback for ChoosePackagesToCheckIn(), continues to bring up UI once source control operations are complete */
 	static void ChoosePackagesToCheckInCallback(const FSourceControlOperationRef& InOperation, ECommandResult::Type InResult);
 
diff --git a/Engine/Source/Editor/UnrealEd/Classes/Editor/EditorEngine.h b/Engine/Source/Editor/UnrealEd/Classes/Editor/EditorEngine.h
--- a/Engine/Source/Editor/UnrealEd/Classes/Editor/EditorEngine.h
+++ b/Engine/Source/Editor/UnrealEd/Classes/Editor/EditorEngine.h
@@ -2594,7 +2594,7 @@ protected:
 	void HandleStageStarted(const FString& InStage, TWeakPtr<SNotificationItem> NotificationItemPtr);
 	void HandleStageCompleted(const FString& InStage, double StageTime, bool bHasCode, TWeakPtr<SNotificationItem> NotificationItemPtr);
 	void HandleLaunchCanceled(double TotalTime, bool bHasCode, TWeakPtr<SNotificationItem> NotificationItemPtr);
-	void HandleLaunchCompleted(bool Succeeded, double TotalTime, bool bHasCode, TWeakPtr<SNotificationItem> NotificationItemPtr);
+	void HandleLaunchCompleted(bool Succeeded, double TotalTime, int32 ErrorCode, bool bHasCode, TWeakPtr<SNotificationItem> NotificationItemPtr);
 
 private:
 	/** Handler for when any asset is loaded in the editor */
diff --git a/Engine/Source/Editor/UnrealEd/Private/Analytics/EditorAnalytics.cpp b/Engine/Source/Editor/UnrealEd/Private/Analytics/EditorAnalytics.cpp
new file mode 100644
--- /dev/null
+++ b/Engine/Source/Editor/UnrealEd/Private/Analytics/EditorAnalytics.cpp
@@ -0,0 +1,154 @@
+// Copyright 1998-2014 Epic Games, Inc. All Rights Reserved.
+
+#include "UnrealEd.h"
+#include "TargetPlatform.h"
+
+#include "EngineAnalytics.h"
+#include "Runtime/Analytics/Analytics/Public/Interfaces/IAnalyticsProvider.h"
+
+#include "EditorAnalytics.h"
+
+void FEditorAnalytics::ReportBuildRequirementsFailure(FString EventName, FString PlatformName, bool bHasCode, int32 Requirements)
+{
+	TArray<FAnalyticsEventAttribute> ParamArray;
+	ParamArray.Add(FAnalyticsEventAttribute(TEXT("Time"), 0.0));
+	if (Requirements & ETargetPlatformReadyStatus::SDKNotFound)
+	{
+		ReportEvent(EventName, PlatformName, bHasCode, EAnalyticsErrorCodes::SDKNotFound, ParamArray);
+	}
+	if (Requirements & ETargetPlatformReadyStatus::ProvisionNotFound)
+	{
+		ReportEvent(EventName, PlatformName, bHasCode, EAnalyticsErrorCodes::ProvisionNotFound, ParamArray);
+	}
+	if (Requirements & ETargetPlatformReadyStatus::SigningKeyNotFound)
+	{
+		ReportEvent(EventName, PlatformName, bHasCode, EAnalyticsErrorCodes::CertificateNotFound, ParamArray);
+	}
+	if (Requirements & ETargetPlatformReadyStatus::CodeUnsupported)
+	{
+		ReportEvent(EventName, PlatformName, bHasCode, EAnalyticsErrorCodes::CodeUnsupported, ParamArray);
+	}
+	if (Requirements & ETargetPlatformReadyStatus::PluginsUnsupported)
+	{
+		ReportEvent(EventName, PlatformName, bHasCode, EAnalyticsErrorCodes::PluginsUnsupported, ParamArray);
+	}
+}
+
+void FEditorAnalytics::ReportEvent(FString EventName, FString PlatformName, bool bHasCode)
+{
+	if( FEngineAnalytics::IsAvailable() )
+	{
+		const UGeneralProjectSettings& ProjectSettings = *GetDefault<UGeneralProjectSettings>();
+		TArray<FAnalyticsEventAttribute> ParamArray;
+		ParamArray.Add(FAnalyticsEventAttribute(TEXT("ProjectID"), ProjectSettings.ProjectID.ToString()));
+		ParamArray.Add(FAnalyticsEventAttribute(TEXT("Platform"), PlatformName));
+		ParamArray.Add(FAnalyticsEventAttribute(TEXT("ProjectType"), bHasCode ? TEXT("C++ Code") : TEXT("Content Only")));
+
+		FEngineAnalytics::GetProvider().RecordEvent( EventName, ParamArray );
+	}
+}
+
+void FEditorAnalytics::ReportEvent(FString EventName, FString PlatformName, bool bHasCode, TArray<FAnalyticsEventAttribute>& ExtraParams)
+{
+	if( FEngineAnalytics::IsAvailable() )
+	{
+		const UGeneralProjectSettings& ProjectSettings = *GetDefault<UGeneralProjectSettings>();
+		TArray<FAnalyticsEventAttribute> ParamArray;
+		ParamArray.Add(FAnalyticsEventAttribute(TEXT("ProjectID"), ProjectSettings.ProjectID.ToString()));
+		ParamArray.Add(FAnalyticsEventAttribute(TEXT("Platform"), PlatformName));
+		ParamArray.Add(FAnalyticsEventAttribute(TEXT("ProjectType"), bHasCode ? TEXT("C++ Code") : TEXT("Content Only")));
+		ParamArray.Append(ExtraParams);
+
+		FEngineAnalytics::GetProvider().RecordEvent( EventName, ParamArray );
+	}
+}
+
+void FEditorAnalytics::ReportEvent(FString EventName, FString PlatformName, bool bHasCode, int32 ErrorCode, TArray<FAnalyticsEventAttribute>& ExtraParams)
+{
+	if( FEngineAnalytics::IsAvailable() )
+	{
+		const UGeneralProjectSettings& ProjectSettings = *GetDefault<UGeneralProjectSettings>();
+		TArray<FAnalyticsEventAttribute> ParamArray;
+		ParamArray.Add(FAnalyticsEventAttribute(TEXT("ProjectID"), ProjectSettings.ProjectID.ToString()));
+		ParamArray.Add(FAnalyticsEventAttribute(TEXT("Platform"), PlatformName));
+		ParamArray.Add(FAnalyticsEventAttribute(TEXT("ProjectType"), bHasCode ? TEXT("C++ Code") : TEXT("Content Only")));
+		ParamArray.Add(FAnalyticsEventAttribute(TEXT("ErrorCode"), ErrorCode));
+		ParamArray.Add(FAnalyticsEventAttribute(TEXT("ErrorName"), TranslateErrorCode(ErrorCode)));
+		ParamArray.Append(ExtraParams);
+
+		FEngineAnalytics::GetProvider().RecordEvent( EventName, ParamArray );
+	}
+}
+
+FString FEditorAnalytics::TranslateErrorCode(int32 ErrorCode)
+{
+	switch (ErrorCode)
+	{
+	case EAnalyticsErrorCodes::UATNotFound:
+		return TEXT("UAT Not Found");
+	case EAnalyticsErrorCodes::Unknown:
+		return TEXT("Unknown Error");
+	case EAnalyticsErrorCodes::Arguments:
+		return TEXT("Invalid Arguments");
+	case EAnalyticsErrorCodes::UnknownCommand:
+		return TEXT("Unknown Command");
+	case EAnalyticsErrorCodes::SDKNotFound:
+		return TEXT("SDK Not Found");
+	case EAnalyticsErrorCodes::ProvisionNotFound:
+		return TEXT("Provision Not Found");
+	case EAnalyticsErrorCodes::CertificateNotFound:
+		return TEXT("Certificate Not Found");
+	case EAnalyticsErrorCodes::ManifestNotFound:
+		return TEXT("Manifest Not Found");
+	case EAnalyticsErrorCodes::KeyNotFound:
+		return TEXT("Key Not Found in Manifest");
+	case EAnalyticsErrorCodes::ProvisionExpired:
+		return TEXT("Provision Has Expired");
+	case EAnalyticsErrorCodes::CertificateExpired:
+		return TEXT("Certificate Has Expired");
+	case EAnalyticsErrorCodes::CertificateProvisionMismatch:
+		return TEXT("Certificate Doesn't Match Provision");
+	case EAnalyticsErrorCodes::LauncherFailed:
+		return TEXT("LauncherWorker Failed to Launch");
+	case EAnalyticsErrorCodes::UATLaunchFailure:
+		return TEXT("UAT Failed to Launch");
+	case EAnalyticsErrorCodes::UnknownCookFailure:
+		return TEXT("Unknown Cook Failure");
+	case EAnalyticsErrorCodes::UnknownDeployFailure:
+		return TEXT("Unknown Deploy Failure");
+	case EAnalyticsErrorCodes::UnknownBuildFailure:
+		return TEXT("Unknown Build Failure");
+	case EAnalyticsErrorCodes::UnknownPackageFailure:
+		return TEXT("Unknown Package Failure");
+	case EAnalyticsErrorCodes::UnknownLaunchFailure:
+		return TEXT("Unknown Launch Failure");
+	case EAnalyticsErrorCodes::StageMissingFile:
+		return TEXT("Could not find file for staging");
+	case EAnalyticsErrorCodes::FailedToCreateIPA:
+		return TEXT("Failed to Create IPA");
+	case EAnalyticsErrorCodes::FailedToCodeSign:
+		return TEXT("Failed to Code Sign");
+	case EAnalyticsErrorCodes::DeviceBackupFailed:
+		return TEXT("Failed to backup device");
+	case EAnalyticsErrorCodes::IPAUninstallFailed:
+		return TEXT("Failed to Uninstall IPA");
+	case EAnalyticsErrorCodes::IPAInstallFailed:
+		return TEXT("Failed to Install IPA");
+	case EAnalyticsErrorCodes::IPANotFound:
+		return TEXT("IPA Not Found for Install");
+	case EAnalyticsErrorCodes::StubNotSignedCorrectly:
+		return TEXT("Stub not signed correctly.");
+	case EAnalyticsErrorCodes::IPAMissingInfoPList:
+		return TEXT("Failed to find Info.plist in IPA");
+	case EAnalyticsErrorCodes::DeleteFile:
+		return TEXT("Could not delete file");
+	case EAnalyticsErrorCodes::DeleteDirectory:
+		return TEXT("Could not delete directory");
+	case EAnalyticsErrorCodes::CreateDirectory:
+		return TEXT("Could not create directory");
+	case EAnalyticsErrorCodes::CopyFile:
+		return TEXT("Could not copy file");
+	}
+	return TEXT("Unknown Error");
+}
+
diff --git a/Engine/Source/Editor/UnrealEd/Private/Kismet2/DebuggerCommands.cpp b/Engine/Source/Editor/UnrealEd/Private/Kismet2/DebuggerCommands.cpp
--- a/Engine/Source/Editor/UnrealEd/Private/Kismet2/DebuggerCommands.cpp
+++ b/Engine/Source/Editor/UnrealEd/Private/Kismet2/DebuggerCommands.cpp
@@ -24,6 +24,7 @@
 #include "Editor/LevelEditor/Public/LevelEditor.h"
 #include "Editor/LevelEditor/Public/ILevelViewport.h"
 
+#include "EditorAnalytics.h"
 
 #define LOCTEXT_NAMESPACE "DebuggerCommands"
 
@@ -1371,18 +1372,32 @@ bool FInternalPlayWorldCommandCallbacks::IsReadyToLaunchOnDevice(FString DeviceI
 	DeviceId.FindChar(TEXT('@'), Index);
 	FString PlatformName = DeviceId.Left(Index);
 
-	// does the project have any code?
+	const PlatformInfo::FPlatformInfo* const PlatformInfo = PlatformInfo::FindPlatformInfo(FName(*PlatformName));
+	check(PlatformInfo);
+
 	FGameProjectGenerationModule& GameProjectModule = FModuleManager::LoadModuleChecked<FGameProjectGenerationModule>(TEXT("GameProjectGeneration"));
-	bool bProjectHasCode = GameProjectModule.Get().GetProjectCodeFileCount() > 0;
+	bool bHasCode = GameProjectModule.Get().GetProjectCodeFileCount() > 0;
+
+	if (PlatformInfo->SDKStatus == PlatformInfo::EPlatformSDKStatus::NotInstalled)
+	{
+		IMainFrameModule& MainFrameModule = FModuleManager::GetModuleChecked<IMainFrameModule>(TEXT("MainFrame"));
+		MainFrameModule.BroadcastMainFrameSDKNotInstalled(PlatformInfo->TargetPlatformName.ToString(), PlatformInfo->SDKTutorial);
+		TArray<FAnalyticsEventAttribute> ParamArray;
+		ParamArray.Add(FAnalyticsEventAttribute(TEXT("Time"), 0.0));
+		FEditorAnalytics::ReportEvent(TEXT("Editor.LaunchOn.Failed"), PlatformInfo->TargetPlatformName.ToString(), bHasCode, EAnalyticsErrorCodes::SDKNotFound, ParamArray);
+		return false;
+	}
 
 	const ITargetPlatform* Platform = GetTargetPlatformManager()->FindTargetPlatform(PlatformName);
 	if (Platform)
 	{
 		FString NotInstalledDocLink;
 		FString ProjectPath = FPaths::IsProjectFilePathSet() ? FPaths::ConvertRelativePathToFull(FPaths::GetProjectFilePath()) : FPaths::RootDir() / FApp::GetGameName() / FApp::GetGameName() + TEXT(".uproject");
-		switch (Platform->IsReadyToBuild(ProjectPath, bProjectHasCode, NotInstalledDocLink))
+		int32 Result = Platform->DoesntHaveRequirements(ProjectPath, bHasCode, NotInstalledDocLink);
+		FEditorAnalytics::ReportBuildRequirementsFailure(TEXT("Editor.LaunchOn.Failed"), PlatformName, bHasCode, Result);
+		switch (Result)
 		{
-			// broadcast this, and assume someone will pick it up
+		// broadcast this, and assume someone will pick it up
 		case ETargetPlatformReadyStatus::SDKNotFound:
 		case ETargetPlatformReadyStatus::ProvisionNotFound:
 		case ETargetPlatformReadyStatus::SigningKeyNotFound:
diff --git a/Engine/Source/Editor/UnrealEd/Private/PlayLevel.cpp b/Engine/Source/Editor/UnrealEd/Private/PlayLevel.cpp
--- a/Engine/Source/Editor/UnrealEd/Private/PlayLevel.cpp
+++ b/Engine/Source/Editor/UnrealEd/Private/PlayLevel.cpp
@@ -30,6 +30,7 @@
 #include "Runtime/Analytics/Analytics/Public/Interfaces/IAnalyticsProvider.h"
 #include "EngineAnalytics.h"
 #include "Engine/GameInstance.h"
+#include "EditorAnalytics.h"
 
 #include "Online.h"
 
@@ -1160,18 +1161,10 @@ void UEditorEngine::HandleStageStarted(const FString& InStage, TWeakPtr<SNotific
 void UEditorEngine::HandleStageCompleted(const FString& InStage, double StageTime, bool bHasCode, TWeakPtr<SNotificationItem> NotificationItemPtr)
 {
 	// analytics for launch on
-	if( FEngineAnalytics::IsAvailable() )
-	{
-		const UGeneralProjectSettings& ProjectSettings = *GetDefault<UGeneralProjectSettings>();
-		TArray<FAnalyticsEventAttribute> ParamArray;
-		ParamArray.Add(FAnalyticsEventAttribute(TEXT("ProjectID"), ProjectSettings.ProjectID.ToString()));
-		ParamArray.Add(FAnalyticsEventAttribute(TEXT("Platform"), PlayUsingLauncherDeviceId.Left(PlayUsingLauncherDeviceId.Find(TEXT("@")))));
-		ParamArray.Add(FAnalyticsEventAttribute(TEXT("ProjectType"), bHasCode ? TEXT("C++ Code") : TEXT("Content Only")));
-		ParamArray.Add(FAnalyticsEventAttribute(TEXT("Time"), StageTime));
-		ParamArray.Add(FAnalyticsEventAttribute(TEXT("StageName"), InStage));
-
-		FEngineAnalytics::GetProvider().RecordEvent( TEXT( "Editor.LaunchOn.StageComplete" ), ParamArray );
-	}
+	TArray<FAnalyticsEventAttribute> ParamArray;
+	ParamArray.Add(FAnalyticsEventAttribute(TEXT("Time"), StageTime));
+	ParamArray.Add(FAnalyticsEventAttribute(TEXT("StageName"), InStage));
+	FEditorAnalytics::ReportEvent(TEXT( "Editor.LaunchOn.StageComplete" ), PlayUsingLauncherDeviceId.Left(PlayUsingLauncherDeviceId.Find(TEXT("@"))), bHasCode, ParamArray);
 }
 
 void UEditorEngine::HandleLaunchCanceled(double TotalTime, bool bHasCode, TWeakPtr<SNotificationItem> NotificationItemPtr)
@@ -1183,22 +1176,14 @@ void UEditorEngine::HandleLaunchCanceled(double TotalTime, bool bHasCode, TWeakP
 	);
 
 	// analytics for launch on
-	if( FEngineAnalytics::IsAvailable() )
-	{
-		const UGeneralProjectSettings& ProjectSettings = *GetDefault<UGeneralProjectSettings>();
-		TArray<FAnalyticsEventAttribute> ParamArray;
-		ParamArray.Add(FAnalyticsEventAttribute(TEXT("ProjectID"), ProjectSettings.ProjectID.ToString()));
-		ParamArray.Add(FAnalyticsEventAttribute(TEXT("Platform"), PlayUsingLauncherDeviceId.Left(PlayUsingLauncherDeviceId.Find(TEXT("@")))));
-		ParamArray.Add(FAnalyticsEventAttribute(TEXT("ProjectType"), bHasCode ? TEXT("C++ Code") : TEXT("Content Only")));
-		ParamArray.Add(FAnalyticsEventAttribute(TEXT("Time"), TotalTime));
-
-		FEngineAnalytics::GetProvider().RecordEvent( TEXT( "Editor.LaunchOn.Canceled" ), ParamArray );
-	}
+	TArray<FAnalyticsEventAttribute> ParamArray;
+	ParamArray.Add(FAnalyticsEventAttribute(TEXT("Time"), TotalTime));
+	FEditorAnalytics::ReportEvent(TEXT( "Editor.LaunchOn.Canceled" ), PlayUsingLauncherDeviceId.Left(PlayUsingLauncherDeviceId.Find(TEXT("@"))), bHasCode, ParamArray);
 
 	bPlayUsingLauncher = false;	
 }
 
-void UEditorEngine::HandleLaunchCompleted(bool Succeeded, double TotalTime, bool bHasCode, TWeakPtr<SNotificationItem> NotificationItemPtr)
+void UEditorEngine::HandleLaunchCompleted(bool Succeeded, double TotalTime, int32 ErrorCode, bool bHasCode, TWeakPtr<SNotificationItem> NotificationItemPtr)
 {
 	if (Succeeded)
 	{
@@ -1222,17 +1207,9 @@ void UEditorEngine::HandleLaunchCompleted(bool Succeeded, double TotalTime, bool
 		);
 
 		// analytics for launch on
-		if( FEngineAnalytics::IsAvailable() )
-		{
-			const UGeneralProjectSettings& ProjectSettings = *GetDefault<UGeneralProjectSettings>();
-			TArray<FAnalyticsEventAttribute> ParamArray;
-			ParamArray.Add(FAnalyticsEventAttribute(TEXT("ProjectID"), ProjectSettings.ProjectID.ToString()));
-			ParamArray.Add(FAnalyticsEventAttribute(TEXT("Platform"), PlayUsingLauncherDeviceId.Left(PlayUsingLauncherDeviceId.Find(TEXT("@")))));
-			ParamArray.Add(FAnalyticsEventAttribute(TEXT("ProjectType"), bHasCode ? TEXT("C++ Code") : TEXT("Content Only")));
-			ParamArray.Add(FAnalyticsEventAttribute(TEXT("Time"), TotalTime));
-
-			FEngineAnalytics::GetProvider().RecordEvent( TEXT( "Editor.LaunchOn.Completed" ), ParamArray );
-		}
+		TArray<FAnalyticsEventAttribute> ParamArray;
+		ParamArray.Add(FAnalyticsEventAttribute(TEXT("Time"), TotalTime));
+		FEditorAnalytics::ReportEvent(TEXT( "Editor.LaunchOn.Completed" ), PlayUsingLauncherDeviceId.Left(PlayUsingLauncherDeviceId.Find(TEXT("@"))), bHasCode, ParamArray);
 	}
 	else
 	{
@@ -1251,19 +1228,9 @@ void UEditorEngine::HandleLaunchCompleted(bool Succeeded, double TotalTime, bool
 			SNotificationItem::CS_Fail,
 			CompletionMsg
 			);
-
-		// analytics for launch on
-		if( FEngineAnalytics::IsAvailable() )
-		{
-			const UGeneralProjectSettings& ProjectSettings = *GetDefault<UGeneralProjectSettings>();
-			TArray<FAnalyticsEventAttribute> ParamArray;
-			ParamArray.Add(FAnalyticsEventAttribute(TEXT("ProjectID"), ProjectSettings.ProjectID.ToString()));
-			ParamArray.Add(FAnalyticsEventAttribute(TEXT("Platform"), PlayUsingLauncherDeviceId.Left(PlayUsingLauncherDeviceId.Find(TEXT("@")))));
-			ParamArray.Add(FAnalyticsEventAttribute(TEXT("ProjectType"), bHasCode ? TEXT("C++ Code") : TEXT("Content Only")));
-			ParamArray.Add(FAnalyticsEventAttribute(TEXT("Time"), TotalTime));
-
-			FEngineAnalytics::GetProvider().RecordEvent( TEXT( "Editor.LaunchOn.Failed" ), ParamArray );
-		}
+		TArray<FAnalyticsEventAttribute> ParamArray;
+		ParamArray.Add(FAnalyticsEventAttribute(TEXT("Time"), TotalTime));
+		FEditorAnalytics::ReportEvent(TEXT( "Editor.LaunchOn.Failed" ), PlayUsingLauncherDeviceId.Left(PlayUsingLauncherDeviceId.Find(TEXT("@"))), bHasCode, ErrorCode, ParamArray);
 	}
 	bPlayUsingLauncher = false;
 }
@@ -1366,16 +1333,8 @@ void UEditorEngine::PlayUsingLauncher()
 		}
 
 		// analytics for launch on
-		if( FEngineAnalytics::IsAvailable() )
-		{
-			const UGeneralProjectSettings& ProjectSettings = *GetDefault<UGeneralProjectSettings>();
-			TArray<FAnalyticsEventAttribute> ParamArray;
-			ParamArray.Add(FAnalyticsEventAttribute(TEXT("ProjectID"), ProjectSettings.ProjectID.ToString()));
-			ParamArray.Add(FAnalyticsEventAttribute(TEXT("Platform"), PlayUsingLauncherDeviceId.Left(PlayUsingLauncherDeviceId.Find(TEXT("@")))));
-			ParamArray.Add(FAnalyticsEventAttribute(TEXT("ProjectType"), bHasCode ? TEXT("C++ Code") : TEXT("Content Only")));
-
-			FEngineAnalytics::GetProvider().RecordEvent( TEXT( "Editor.LaunchOn.Started" ), ParamArray );
-		}
+		int32 ErrorCode = 0;
+		FEditorAnalytics::ReportEvent(TEXT( "Editor.LaunchOn.Started" ), PlayUsingLauncherDeviceId.Left(PlayUsingLauncherDeviceId.Find(TEXT("@"))), bHasCode);
 
 		NotificationItem->SetCompletionState(SNotificationItem::CS_Pending);
 
@@ -1400,17 +1359,9 @@ void UEditorEngine::PlayUsingLauncher()
 			bPlayUsingLauncher = false;
 
 			// analytics for launch on
-			if( FEngineAnalytics::IsAvailable() )
-			{
-				const UGeneralProjectSettings& ProjectSettings = *GetDefault<UGeneralProjectSettings>();
-				TArray<FAnalyticsEventAttribute> ParamArray;
-				ParamArray.Add(FAnalyticsEventAttribute(TEXT("ProjectID"), ProjectSettings.ProjectID.ToString()));
-				ParamArray.Add(FAnalyticsEventAttribute(TEXT("Platform"), PlayUsingLauncherDeviceId.Left(PlayUsingLauncherDeviceId.Find(TEXT("@")))));
-				ParamArray.Add(FAnalyticsEventAttribute(TEXT("ProjectType"), bHasCode ? TEXT("C++ Code") : TEXT("Content Only")));
-				ParamArray.Add(FAnalyticsEventAttribute(TEXT("Time"), 0.0f));
-
-				FEngineAnalytics::GetProvider().RecordEvent( TEXT( "Editor.LaunchOn.Failed" ), ParamArray );
-			}
+			TArray<FAnalyticsEventAttribute> ParamArray;
+			ParamArray.Add(FAnalyticsEventAttribute(TEXT("Time"), 0.0));
+			FEditorAnalytics::ReportEvent(TEXT( "Editor.LaunchOn.Failed" ), PlayUsingLauncherDeviceId.Left(PlayUsingLauncherDeviceId.Find(TEXT("@"))), bHasCode, EAnalyticsErrorCodes::LauncherFailed, ParamArray );
 		}
 	}
 }
diff --git a/Engine/Source/Editor/UnrealEd/Public/EditorAnalytics.h b/Engine/Source/Editor/UnrealEd/Public/EditorAnalytics.h
new file mode 100644
--- /dev/null
+++ b/Engine/Source/Editor/UnrealEd/Public/EditorAnalytics.h
@@ -0,0 +1,97 @@
+// Copyright 1998-2014 Epic Games, Inc. All Rights Reserved.
+
+#pragma once
+
+#include "Runtime/Analytics/Analytics/Public/Interfaces/IAnalyticsProvider.h"
+#include "Runtime/Engine/Public/EngineAnalytics.h"
+
+namespace EAnalyticsErrorCodes
+{
+	/**
+	 * Error Codes for a variety of tool errors
+	 * NOTE: this needs to be kept in sync with iPhonePackager.cs and AutomationTool/Program.cs
+	 * Any new error added here needs a text string in TranslateErrorCode
+	 */
+	enum Type
+	{
+		UATNotFound = -1,
+		Success = 0,
+		Unknown = 1,
+		Arguments = 2,
+		UnknownCommand = 3,
+		SDKNotFound = 10,
+		ProvisionNotFound = 11,
+		CertificateNotFound = 12,
+		ManifestNotFound = 14,
+		KeyNotFound = 15,
+		ProvisionExpired = 16,
+		CertificateExpired = 17,
+		CertificateProvisionMismatch = 18,
+		CodeUnsupported = 19,
+		PluginsUnsupported = 20,
+		UnknownCookFailure = 25,
+		UnknownDeployFailure = 26,
+		UnknownBuildFailure = 27,
+		UnknownPackageFailure = 28,
+		UnknownLaunchFailure = 29,
+		StageMissingFile = 30,
+		FailedToCreateIPA = 31,
+		FailedToCodeSign = 32,
+		DeviceBackupFailed = 33,
+		IPAUninstallFailed = 34,
+		IPAInstallFailed = 35,
+		IPANotFound = 36,
+		StubNotSignedCorrectly = 37,
+		IPAMissingInfoPList = 38,
+		DeleteFile = 39,
+		DeleteDirectory = 40,
+		CreateDirectory = 41,
+		CopyFile = 42,
+		LauncherFailed = 100,
+		UATLaunchFailure = 101,
+	};
+};
+
+class FEditorAnalytics
+{
+public:
+	/*
+	* Reports an event to the analytics system if it is enabled
+	*
+	* @param EventName - name of the event
+	* @param PlatformName - name of the platform being used
+	*/
+	UNREALED_API static void ReportEvent(FString EventName, FString PlatformName, bool bHasCode);
+
+	/*
+	* Reports an event to the analytics system if it is enabled with some extra parameters
+	*
+	* @param EventName - name of the event
+	* @param PlatformName - name of the platform being used
+	* @param ExtraParams - extra data needed by the event
+	*/
+	UNREALED_API static void ReportEvent(FString EventName, FString PlatformName, bool bHasCode, TArray<FAnalyticsEventAttribute>& ExtraParams);
+
+	/*
+	* Reports an event to the analytics system if it is enabled with some extra parameters with an error code
+	*
+	* @param EventName - name of the event
+	* @param PlatformName - name of the platform being used
+	* @param ErrorCode - error code of the event
+	* @param ExtraParams - extra data needed by the event
+	*/
+	UNREALED_API static void ReportEvent(FString EventName, FString PlatformName, bool bHasCode, int32 ErrorCode, TArray<FAnalyticsEventAttribute>& ExtraParams);
+
+	/*
+	* Reports an event to the analytics system of a build requirements failure
+	*
+	* @param EventName - name of the event
+	* @param PlatformName - name of the platform being used
+	* @param Requirements - requirements that didn't pass
+	*/
+	UNREALED_API static void ReportBuildRequirementsFailure(FString EventName, FString PlatformName, bool bHasCode, int32 Requirements);
+
+private:
+	/* Translates the error code to something human readable */
+	static FString TranslateErrorCode(int32 ErrorCode);
+};
\ No newline at end of file
diff --git a/Engine/Source/Editor/UnrealEd/Public/UnrealEd.h b/Engine/Source/Editor/UnrealEd/Public/UnrealEd.h
--- a/Engine/Source/Editor/UnrealEd/Public/UnrealEd.h
+++ b/Engine/Source/Editor/UnrealEd/Public/UnrealEd.h
@@ -56,6 +56,8 @@ UNREALED_API void EditorExit();
 
 #include "UnrealEdMessages.h"
 
+#include "EditorAnalytics.h"
+
 #endif // __cplusplus
 
 #endif	// __UnrealEd_h__
diff --git a/Engine/Source/Programs/AutomationTool/DeploymentContext.cs b/Engine/Source/Programs/AutomationTool/DeploymentContext.cs
--- a/Engine/Source/Programs/AutomationTool/DeploymentContext.cs
+++ b/Engine/Source/Programs/AutomationTool/DeploymentContext.cs
@@ -416,6 +416,7 @@ public class DeploymentContext //: ProjectParams
 
 		if (FilesAdded == 0 && !bAllowNone)
 		{
+			AutomationTool.ErrorReporter.Error(String.Format("No files found to deploy for {0} with wildcard {1} and exclusions {2}", InPath, Wildcard, ExcludeWildcard), (int)AutomationTool.ErrorCodes.Error_StageMissingFile);
 			throw new AutomationException("No files found to deploy for {0} with wildcard {1} and exclusions {2}", InPath, Wildcard, ExcludeWildcard);
 		}
 
diff --git a/Engine/Source/Programs/AutomationTool/IOS/IOSPlatform.Automation.cs b/Engine/Source/Programs/AutomationTool/IOS/IOSPlatform.Automation.cs
--- a/Engine/Source/Programs/AutomationTool/IOS/IOSPlatform.Automation.cs
+++ b/Engine/Source/Programs/AutomationTool/IOS/IOSPlatform.Automation.cs
@@ -275,6 +275,7 @@ public class IOSPlatform : Platform
 			// verify the .ipa exists
 			if (!FileExists(ProjectIPA))
 			{
+				ErrorReporter.Error(String.Format("PACKAGE FAILED - {0} was not created", ProjectIPA), (int)ErrorCodes.Error_FailedToCreateIPA);
 				throw new AutomationException("PACKAGE FAILED - {0} was not created", ProjectIPA);
 			}
 
@@ -493,6 +494,7 @@ public class IOSPlatform : Platform
 		}
 		if (Result.ExitCode != 0)
 		{
+			ErrorReporter.Error("CodeSign Failed", (int)ErrorCodes.Error_FailedToCodeSign);
 			throw new AutomationException("CodeSign Failed");
 		}
 	}
diff --git a/Engine/Source/Programs/AutomationTool/ProcessUtils.cs b/Engine/Source/Programs/AutomationTool/ProcessUtils.cs
--- a/Engine/Source/Programs/AutomationTool/ProcessUtils.cs
+++ b/Engine/Source/Programs/AutomationTool/ProcessUtils.cs
@@ -772,6 +772,8 @@ namespace AutomationTool
 
             if (Result > MaxSuccessCode || Result < 0)
             {
+				ErrorReporter.Error(String.Format("Command failed (Result:{3}): {0} {1}. See logfile for details: '{2}' ",
+												App, CommandLine, Path.GetFileName(Logfile), Result.ExitCode), Result.ExitCode);
                 throw new AutomationException(String.Format("Command failed (Result:{3}): {0} {1}. See logfile for details: '{2}' ",
                                                 App, CommandLine, Path.GetFileName(Logfile), Result.ExitCode));
             }
diff --git a/Engine/Source/Programs/AutomationTool/Program.cs b/Engine/Source/Programs/AutomationTool/Program.cs
--- a/Engine/Source/Programs/AutomationTool/Program.cs
+++ b/Engine/Source/Programs/AutomationTool/Program.cs
@@ -11,10 +11,67 @@ using System.Reflection;
 
 namespace AutomationTool
 {
+	// NOTE: this needs to be kept in sync with EditorAnalytics.h and iPhonePackager.cs
+	public enum ErrorCodes
+	{
+		Error_UATNotFound = -1,
+		Error_Success = 0,
+		Error_Unknown = 1,
+		Error_Arguments = 2,
+		Error_UnknownCommand = 3,
+		Error_SDKNotFound = 10,
+		Error_ProvisionNotFound = 11,
+		Error_CertificateNotFound = 12,
+		Error_ProvisionAndCertificateNotFound = 13,
+		Error_InfoPListNotFound = 14,
+		Error_KeyNotFoundInPList = 15,
+		Error_ProvisionExpired = 16,
+		Error_CertificateExpired = 17,
+		Error_CertificateProvisionMismatch = 18,
+		Error_CodeUnsupported = 19,
+		Error_PluginsUnsupported = 20,
+		Error_UnknownCookFailure = 25,
+		Error_UnknownDeployFailure = 26,
+		Error_UnknownBuildFailure = 27,
+		Error_UnknownPackageFailure = 28,
+		Error_UnknownLaunchFailure = 29,
+		Error_StageMissingFile = 30,
+		Error_FailedToCreateIPA = 31,
+		Error_FailedToCodeSign = 32,
+		Error_DeviceBackupFailed = 33,
+		Error_IPAUninstallFailed = 34,
+		Error_IPAInstallFailed = 35,
+		Error_IPANotFound = 36,
+		Error_StubNotSignedCorrectly = 37,
+		Error_IPAMissingInfoPList = 38,
+		Error_DeleteFile = 39,
+		Error_DeleteDirectory = 40,
+		Error_CreateDirectory = 41,
+		Error_CopyFile = 42,
+		Error_LauncherFailed = 100,
+		Error_UATLaunchFailure = 101,
+	};
+
+	public class ErrorReporter
+	{
+		static public void Error(string Line, int Code = 1)
+		{
+			if (Program.ReturnCode == 0)
+			{
+				Program.ReturnCode = Code;
+			}
+
+			Console.ForegroundColor = ConsoleColor.Red;
+			Log.WriteLine(TraceEventType.Error, "AutomationTool error: " + Line);
+			Console.ResetColor();
+		}
+	};
+
 	class Program
 	{
 		// This needs to be static, otherwise SetConsoleCtrlHandler will result in a crash on exit.
 		static ProcessManager.CtrlHandlerDelegate ProgramCtrlHandler = new ProcessManager.CtrlHandlerDelegate(CtrlHandler);
+		static public int ReturnCode = 0;
 
 		[STAThread]
 		static int Main()
@@ -22,7 +79,6 @@ namespace AutomationTool
 			var CommandLine = SharedUtils.ParseCommandLine();
 			HostPlatform.Initialize();
 			
-			int Result = 1;
 			LogUtils.InitLogging(CommandLine);
 			Log.WriteLine(TraceEventType.Information, "Running on {0}", HostPlatform.Current.GetType().Name);
 
@@ -48,18 +104,22 @@ namespace AutomationTool
 			try
 			{
 				// Don't allow simultaneous execution of AT (in the same branch)
-				Result = InternalUtils.RunSingleInstance(MainProc, CommandLine);
+				ReturnCode = InternalUtils.RunSingleInstance(MainProc, CommandLine);
 			}
 			catch (Exception Ex)
 			{
 				Log.WriteLine(TraceEventType.Error, "AutomationTool terminated with exception:");
 				Log.WriteLine(TraceEventType.Error, LogUtils.FormatException(Ex));
 				Log.WriteLine(TraceEventType.Error, Ex.Message);
+				if (ReturnCode == 0)
+				{
+					ReturnCode = (int)ErrorCodes.Error_Unknown;
+				}
 			}
 
 			// Make sure there's no directiories on the stack.
 			CommandUtils.ClearDirStack();
-			Environment.ExitCode = Result;
+			Environment.ExitCode = ReturnCode;
 
 			// Try to kill process before app domain exits to leave the other KillAll call to extreme edge cases
 			if (ShouldKillProcesses)
@@ -67,10 +127,10 @@ namespace AutomationTool
 				ProcessManager.KillAll();
 			}
 
-			Log.WriteLine(TraceEventType.Information, "AutomationTool exiting with ExitCode={0}", Result);
+			Log.WriteLine(TraceEventType.Information, "AutomationTool exiting with ExitCode={0}", ReturnCode);
 			LogUtils.CloseFileLogging();
 
-			return Result;
+			return ReturnCode;
 		}
 
 		static bool CtrlHandler(CtrlTypes EventType)
diff --git a/Engine/Source/Programs/AutomationTool/Scripts/CookCommand.Automation.cs b/Engine/Source/Programs/AutomationTool/Scripts/CookCommand.Automation.cs
--- a/Engine/Source/Programs/AutomationTool/Scripts/CookCommand.Automation.cs
+++ b/Engine/Source/Programs/AutomationTool/Scripts/CookCommand.Automation.cs
@@ -90,6 +90,7 @@ public partial class Project : CommandUtils
 			// Delete cooked data (if any) as it may be incomplete / corrupted.
 			Log("Cook failed. Deleting cooked data.");
 			CleanupCookedData(PlatformsToCook, Params);
+			AutomationTool.ErrorReporter.Error("Cook failed.", (int)AutomationTool.ErrorCodes.Error_UnknownCookFailure);
 			throw Ex;
 		}
 	}
diff --git a/Engine/Source/Programs/IOS/iPhonePackager/CookTime.cs b/Engine/Source/Programs/IOS/iPhonePackager/CookTime.cs
--- a/Engine/Source/Programs/IOS/iPhonePackager/CookTime.cs
+++ b/Engine/Source/Programs/IOS/iPhonePackager/CookTime.cs
@@ -72,7 +72,7 @@ namespace iPhonePackager
 			if (!bHasCodeSignature || !bHasMobileProvision)
 			{
 				Program.Error("Stub IPA does not appear to be signed correctly (missing mobileprovision or CodeResources)");
-				Program.ReturnCode = 129;
+				Program.ReturnCode = (int)ErrorCodes.Error_StubNotSignedCorrectly;
 			}
 
             // Set encoding to support unicode filenames
@@ -92,7 +92,7 @@ namespace iPhonePackager
 			if (Zip[PListPathInZIP] == null)
 			{
 				Program.Error("Failed to find Info.plist in IPA (cannot update plist version)");
-				Program.ReturnCode = 142;
+				Program.ReturnCode = (int)ErrorCodes.Error_IPAMissingInfoPList;
 				return "";
 			}
 			else
diff --git a/Engine/Source/Programs/IOS/iPhonePackager/DeployTime.cs b/Engine/Source/Programs/IOS/iPhonePackager/DeployTime.cs
--- a/Engine/Source/Programs/IOS/iPhonePackager/DeployTime.cs
+++ b/Engine/Source/Programs/IOS/iPhonePackager/DeployTime.cs
@@ -113,7 +113,7 @@ namespace iPhonePackager
 						if (!DeploymentHelper.Get().BackupDocumentsDirectory(ApplicationIdentifier, Config.GetRootBackedUpDocumentsDirectory()))
 						{
 							Program.Error("Failed to transfer documents directory from device to PC");
-							Program.ReturnCode = 100;
+							Program.ReturnCode = (int)ErrorCodes.Error_DeviceBackupFailed;
 						}
 					}
 					break;
@@ -128,7 +128,7 @@ namespace iPhonePackager
 						if (!DeploymentHelper.Get().UninstallIPAOnDevice(ApplicationIdentifier))
 						{
 							Program.Error("Failed to uninstall IPA on device");
-							Program.ReturnCode = 100;
+							Program.ReturnCode = (int)ErrorCodes.Error_IPAUninstallFailed;
 						}
 					}
 					break;
@@ -146,13 +146,13 @@ namespace iPhonePackager
 							if (!DeploymentHelper.Get().InstallIPAOnDevice(IPAPath))
 							{
 								Program.Error("Failed to install IPA on device");
-								Program.ReturnCode = 100;
+								Program.ReturnCode = (int)ErrorCodes.Error_IPAInstallFailed;
 							}
 						}
 						else
 						{
 							Program.Error(String.Format("Failed to find IPA file: '{0}'", IPAPath));
-							Program.ReturnCode = 100;
+							Program.ReturnCode = (int)ErrorCodes.Error_IPANotFound;
 						}
 					}
 					break;
diff --git a/Engine/Source/Programs/IOS/iPhonePackager/FileOperations.cs b/Engine/Source/Programs/IOS/iPhonePackager/FileOperations.cs
--- a/Engine/Source/Programs/IOS/iPhonePackager/FileOperations.cs
+++ b/Engine/Source/Programs/IOS/iPhonePackager/FileOperations.cs
@@ -123,7 +123,7 @@ namespace iPhonePackager
 					{
 						Program.Error( "Could not delete file: " + Info.FullName );
 						Program.Error( "    Exception: " + Ex.Message );
-						Program.ReturnCode = 300;
+						Program.ReturnCode = (int)ErrorCodes.Error_DeleteFile;
 					}
 				}
 
@@ -140,7 +140,7 @@ namespace iPhonePackager
 				{
 					Program.Error( "Could not delete folder: " + DirInfo.FullName );
 					Program.Error( "    Exception: " + Ex.Message );
-					Program.ReturnCode = 310;
+					Program.ReturnCode = (int)ErrorCodes.Error_DeleteDirectory;
 				}
 			}
 		}
@@ -162,7 +162,7 @@ namespace iPhonePackager
 				{
 					Program.Error( "Could not delete file: " + Info.FullName );
 					Program.Error( "    Exception: " + Ex.Message );
-					Program.ReturnCode = 310;
+					Program.ReturnCode = (int)ErrorCodes.Error_DeleteFile;
 				}
 			}
 		}
@@ -190,7 +190,7 @@ namespace iPhonePackager
 				{
 					Program.Error( "Could not create folder: " + DestFolderName );
 					Program.Error( "    Exception: " + Ex.Message );
-					Program.ReturnCode = 312;
+					Program.ReturnCode = (int)ErrorCodes.Error_CreateDirectory;
 				}
 				RecursiveFolderCopy( SourceSubFolderInfo, new DirectoryInfo( DestFolderName ), bLowercase );
 			}
@@ -227,7 +227,7 @@ namespace iPhonePackager
 			{
 				Program.Error( "Could not create folder: " + DestFolderInfo.FullName );
 				Program.Error( "    Exception: " + Ex.Message );
-				Program.ReturnCode = 312;
+				Program.ReturnCode = (int)ErrorCodes.Error_CreateDirectory;
 			}
 
 			RecursiveFolderCopy(SourceFolderInfo, DestFolderInfo, bLowercase);
@@ -390,7 +390,7 @@ namespace iPhonePackager
 				if (bCopyMustSucceed)
 				{
 					Program.Error(Msg);
-					Program.ReturnCode = 320;
+					Program.ReturnCode = (int)ErrorCodes.Error_CopyFile;
 				}
 				else
 				{
diff --git a/Engine/Source/Programs/IOS/iPhonePackager/Utilities.cs b/Engine/Source/Programs/IOS/iPhonePackager/Utilities.cs
--- a/Engine/Source/Programs/IOS/iPhonePackager/Utilities.cs
+++ b/Engine/Source/Programs/IOS/iPhonePackager/Utilities.cs
@@ -540,6 +540,7 @@ namespace iPhonePackager
 			if (!Utilities.GetSourcePList(out PListAsString))
 			{
 				Program.Error("Failed to find source PList");
+				Program.ReturnCode = (int)ErrorCodes.Error_InfoPListNotFound;
 				return "(unknown)";
 			}
 
@@ -553,6 +554,7 @@ namespace iPhonePackager
 			else
 			{
 				Program.Error("Failed to find a value for {0} in PList", KeyName);
+				Program.ReturnCode = (int)ErrorCodes.Error_KeyNotFoundInPList;
 				return "(unknown)";
 			}
 		}
@@ -593,6 +595,7 @@ namespace iPhonePackager
 			else
 			{
 				Program.Error(" ... failed to locate the source .plist file");
+				Program.ReturnCode = (int)ErrorCodes.Error_KeyNotFoundInPList;
 				PListSource = "";
 				return false;
 			}
diff --git a/Engine/Source/Programs/IOS/iPhonePackager/iPhonePackager.cs b/Engine/Source/Programs/IOS/iPhonePackager/iPhonePackager.cs
--- a/Engine/Source/Programs/IOS/iPhonePackager/iPhonePackager.cs
+++ b/Engine/Source/Programs/IOS/iPhonePackager/iPhonePackager.cs
@@ -18,6 +18,47 @@ using System.Security.Cryptography.X509Certificates;
 
 namespace iPhonePackager
 {
+	// NOTE: this needs to be kept in sync with EditorAnalytics.h and AutomationTool/Program.cs
+	public enum ErrorCodes
+	{
+		Error_UATNotFound = -1,
+		Error_Success = 0,
+		Error_Unknown = 1,
+		Error_Arguments = 2,
+		Error_UnknownCommand = 3,
+		Error_SDKNotFound = 10,
+		Error_ProvisionNotFound = 11,
+		Error_CertificateNotFound = 12,
+		Error_ProvisionAndCertificateNotFound = 13,
+		Error_InfoPListNotFound = 14,
+		Error_KeyNotFoundInPList = 15,
+		Error_ProvisionExpired = 16,
+		Error_CertificateExpired = 17,
+		Error_CertificateProvisionMismatch = 18,
+		Error_CodeUnsupported = 19,
+		Error_PluginsUnsupported = 20,
+		Error_UnknownCookFailure = 25,
+		Error_UnknownDeployFailure = 26,
+		Error_UnknownBuildFailure = 27,
+		Error_UnknownPackageFailure = 28,
+		Error_UnknownLaunchFailure = 29,
+		Error_StageMissingFile = 30,
+		Error_FailedToCreateIPA = 31,
+		Error_FailedToCodeSign = 32,
+		Error_DeviceBackupFailed = 33,
+		Error_IPAUninstallFailed = 34,
+		Error_IPAInstallFailed = 35,
+		Error_IPANotFound = 36,
+		Error_StubNotSignedCorrectly = 37,
+		Error_IPAMissingInfoPList = 38,
+		Error_DeleteFile = 39,
+		Error_DeleteDirectory = 40,
+		Error_CreateDirectory = 41,
+		Error_CopyFile = 42,
+		Error_LauncherFailed = 100,
+		Error_UATLaunchFailure = 101,
+	};
+
 	public partial class Program
 	{
 		static public int ReturnCode = 0;
@@ -273,12 +314,14 @@ namespace iPhonePackager
 			if (GameName.Length == 0)
 			{
 				Error( "Invalid number of arguments" );
+				Program.ReturnCode = (int)ErrorCodes.Error_Arguments;
 				return false;
 			}
 
 			if (Config.bCreateStubSet && Config.bForDistribution)
 			{
 				Error("-createstub and -distribution are mutually exclusive");
+				Program.ReturnCode = (int)ErrorCodes.Error_Arguments;
 				return false;
 			}
 
@@ -342,6 +385,7 @@ namespace iPhonePackager
 							if (Config.bForDistribution)
 							{
 								Error("configure cannot be used with -distribution");
+								Program.ReturnCode = (int)ErrorCodes.Error_Arguments;
 							}
 							else
 							{
@@ -358,13 +402,16 @@ namespace iPhonePackager
 				if (!bHandledCommand)
 				{
 					Error("Unknown command");
-					ReturnCode = 2;
+					Program.ReturnCode = (int)ErrorCodes.Error_UnknownCommand;
 				}
 			}
 			catch (Exception Ex)
 			{
 				Error( "Error while executing command: " + Ex.ToString() );
-				ReturnCode = 200;
+				if (Program.ReturnCode == 0)
+				{
+					Program.ReturnCode = (int)ErrorCodes.Error_Unknown;
+				}
 			}
 
 			Console.WriteLine();
@@ -522,7 +569,7 @@ namespace iPhonePackager
 					Log("Sample commandlines:");
 					Log(" ... iPhonePackager Deploy UDKGame Release");
 					Log(" ... iPhonePackager RPC SwordGame Shipping MakeApp");
-					return 1;
+					return (int)ErrorCodes.Error_Arguments;
 				}
 
 				Log("Executing iPhonePackager " + String.Join(" ", args));
@@ -546,7 +593,7 @@ namespace iPhonePackager
 				// setup configuration
 				if (!Config.Initialize(InitialCurrentDirectory, GamePath))
 				{
-					return -1;
+					return (int)ErrorCodes.Error_Arguments;
 				}
 
 				switch (MainCommand.ToLowerInvariant())
@@ -556,7 +603,7 @@ namespace iPhonePackager
 						string dllPath = Microsoft.Win32.Registry.GetValue("HKEY_LOCAL_MACHINE\\SOFTWARE\\Wow6432Node\\Apple Inc.\\Apple Mobile Device Support\\Shared", "iTunesMobileDeviceDLL", null) as string; 
 						if (String.IsNullOrEmpty(dllPath) || !File.Exists(dllPath))
 						{
-							Error("iTunes Not Found!!", 10);
+							Error("iTunes Not Found!!", ErrorCodes.Error_SDKNotFound);
 						}
 						else
 						{
@@ -567,19 +614,19 @@ namespace iPhonePackager
 							bool foundPlist = CodeSignatureBuilder.FindRequiredFiles(out Provision, out Cert, out bHasOverrides);
 							if (!foundPlist)
 							{
-								Error("Could not find a valid plist file!!", 14);
+								Error("Could not find a valid plist file!!", ErrorCodes.Error_InfoPListNotFound);
 							}
 							else if (Provision == null && Cert == null)
 							{
-								Error("No Provision or cert found!!", 13);
+								Error("No Provision or cert found!!", ErrorCodes.Error_ProvisionAndCertificateNotFound);
 							}
 							else if (Provision == null)
 							{
-								Error("No Provision found!!", 11);
+								Error("No Provision found!!", ErrorCodes.Error_ProvisionNotFound);
 							}
 							else if (Cert == null)
 							{
-								Error("No Signing Certificate found!!", 12);
+								Error("No Signing Certificate found!!", ErrorCodes.Error_CertificateNotFound);
 							}
 						}
 						break;
@@ -590,6 +637,7 @@ namespace iPhonePackager
 							if (Config.bCreateStubSet)
 							{
 								Error("packageapp cannot be used with the -createstub switch");
+								Program.ReturnCode = (int)ErrorCodes.Error_Arguments;
 							}
 							else
 							{
@@ -605,6 +653,7 @@ namespace iPhonePackager
 							if (Config.bCreateStubSet)
 							{
 								Error("repackagefromstage cannot be used with the -createstub switches");
+								Program.ReturnCode = (int)ErrorCodes.Error_Arguments;
 							}
 							else
 							{
@@ -662,7 +711,10 @@ namespace iPhonePackager
 			catch (Exception Ex)
 			{
 				Error("Application exception: " + Ex.ToString());
-				ReturnCode = 1;
+				if (ReturnCode == 0)
+				{
+					Program.ReturnCode = (int)ErrorCodes.Error_Unknown;
+				}
 			}
 			finally
 			{
