UE4: Fix for loading the mainframe module while in a commandlet. There are no key events when running in a commandlet.

diff --git a/Engine/Source/Editor/MainFrame/Private/Frame/MainFrameActions.cpp b/Engine/Source/Editor/MainFrame/Private/Frame/MainFrameActions.cpp
--- a/Engine/Source/Editor/MainFrame/Private/Frame/MainFrameActions.cpp
+++ b/Engine/Source/Editor/MainFrame/Private/Frame/MainFrameActions.cpp
@@ -27,8 +27,11 @@ FMainFrameCommands::FMainFrameCommands()
 
 void FMainFrameCommands::RegisterCommands()
 {
-	// The global action list was created at static initialization time. Create a handler for otherwise unhandled keyboard input to route key commands through this list.
-	FSlateApplication::Get().SetUnhandledKeyDownEventHandler( FOnKeyboardEvent::CreateStatic( &FMainFrameActionCallbacks::OnUnhandledKeyDownEvent ) );
+	if ( !IsRunningCommandlet() )
+	{
+		// The global action list was created at static initialization time. Create a handler for otherwise unhandled keyboard input to route key commands through this list.
+		FSlateApplication::Get().SetUnhandledKeyDownEventHandler( FOnKeyboardEvent::CreateStatic( &FMainFrameActionCallbacks::OnUnhandledKeyDownEvent ) );
+	}
 
 	// Make a default can execute action that disables input when in debug mode
 	FCanExecuteAction DefaultExecuteAction = FCanExecuteAction::CreateStatic( &FMainFrameActionCallbacks::DefaultCanExecuteAction );
