made it so when distribution is selected for File|Package then Shipping is automatically set and no other configuration is allowed
TTP327507
#ue4

diff --git a/Engine/Source/Editor/MainFrame/Private/Frame/MainFrameActions.cpp b/Engine/Source/Editor/MainFrame/Private/Frame/MainFrameActions.cpp
--- a/Engine/Source/Editor/MainFrame/Private/Frame/MainFrameActions.cpp
+++ b/Engine/Source/Editor/MainFrame/Private/Frame/MainFrameActions.cpp
@@ -469,6 +469,16 @@ void FMainFrameActionCallbacks::PackageBuildConfiguration( EProjectPackagingBuil
 	PackagingSettings->BuildConfiguration = BuildConfiguration;
 }
 
+bool FMainFrameActionCallbacks::CanPackageBuildConfiguration( EProjectPackagingBuildConfigurations BuildConfiguration )
+{
+	UProjectPackagingSettings* PackagingSettings = Cast<UProjectPackagingSettings>(UProjectPackagingSettings::StaticClass()->GetDefaultObject());
+	if (PackagingSettings->ForDistribution && BuildConfiguration != PPBC_Shipping)
+	{
+		return false;
+	}
+	return true;
+}
+
 bool FMainFrameActionCallbacks::PackageBuildConfigurationIsChecked( EProjectPackagingBuildConfigurations BuildConfiguration )
 {
 	return (GetDefault<UProjectPackagingSettings>()->BuildConfiguration == BuildConfiguration);
diff --git a/Engine/Source/Editor/MainFrame/Private/Frame/MainFrameActions.h b/Engine/Source/Editor/MainFrame/Private/Frame/MainFrameActions.h
--- a/Engine/Source/Editor/MainFrame/Private/Frame/MainFrameActions.h
+++ b/Engine/Source/Editor/MainFrame/Private/Frame/MainFrameActions.h
@@ -160,6 +160,9 @@ public:
 	/** Sets the project packaging build configuration. */
 	static void PackageBuildConfiguration( EProjectPackagingBuildConfigurations BuildConfiguration );
 
+	/** Determines if the packaging build configuration can be used. */
+	static bool CanPackageBuildConfiguration( EProjectPackagingBuildConfigurations BuildConfiguration );
+
 	/** Determines whether the specified build configuration option is checked. */
 	static bool PackageBuildConfigurationIsChecked( EProjectPackagingBuildConfigurations BuildConfiguration );
 
diff --git a/Engine/Source/Editor/MainFrame/Private/Menus/PackageProjectMenu.h b/Engine/Source/Editor/MainFrame/Private/Menus/PackageProjectMenu.h
--- a/Engine/Source/Editor/MainFrame/Private/Menus/PackageProjectMenu.h
+++ b/Engine/Source/Editor/MainFrame/Private/Menus/PackageProjectMenu.h
@@ -151,7 +151,7 @@ protected:
 				FSlateIcon(),
 				FUIAction(
 					FExecuteAction::CreateStatic(&FMainFrameActionCallbacks::PackageBuildConfiguration, PPBC_DebugGame),
-					FCanExecuteAction(),
+					FCanExecuteAction::CreateStatic(&FMainFrameActionCallbacks::CanPackageBuildConfiguration, PPBC_DebugGame),
 					FIsActionChecked::CreateStatic(&FMainFrameActionCallbacks::PackageBuildConfigurationIsChecked, PPBC_DebugGame)
 				),
 				NAME_None,
@@ -165,7 +165,7 @@ protected:
 			FSlateIcon(),
 			FUIAction(
 				FExecuteAction::CreateStatic(&FMainFrameActionCallbacks::PackageBuildConfiguration, PPBC_Development),
-				FCanExecuteAction(),
+				FCanExecuteAction::CreateStatic(&FMainFrameActionCallbacks::CanPackageBuildConfiguration, PPBC_Development),
 				FIsActionChecked::CreateStatic(&FMainFrameActionCallbacks::PackageBuildConfigurationIsChecked, PPBC_Development)
 			),
 			NAME_None,
diff --git a/Engine/Source/Editor/UnrealEd/Classes/Settings/ProjectPackagingSettings.h b/Engine/Source/Editor/UnrealEd/Classes/Settings/ProjectPackagingSettings.h
--- a/Engine/Source/Editor/UnrealEd/Classes/Settings/ProjectPackagingSettings.h
+++ b/Engine/Source/Editor/UnrealEd/Classes/Settings/ProjectPackagingSettings.h
@@ -50,8 +50,8 @@ public:
 	bool FullRebuild;
 
 	/**
-	 * If enabled, a distribution build will be created
-	 * If disabled, a development build will be created
+	 * If enabled, a distribution build will be created and the shipping configuration will be used
+	 * If disabled, an development build will be created
 	 * Distribution builds are for publishing to the App Store
 	 */
 	UPROPERTY(config, EditAnywhere, Category=Project, AdvancedDisplay)
@@ -94,4 +94,5 @@ public:
 	// UObject Interface
 
 	virtual void PostEditChangeProperty( struct FPropertyChangedEvent& PropertyChangedEvent ) override;
+	virtual bool CanEditChange( const UProperty* InProperty ) const;
 };
diff --git a/Engine/Source/Editor/UnrealEd/Private/Settings/SettingsClasses.cpp b/Engine/Source/Editor/UnrealEd/Private/Settings/SettingsClasses.cpp
--- a/Engine/Source/Editor/UnrealEd/Private/Settings/SettingsClasses.cpp
+++ b/Engine/Source/Editor/UnrealEd/Private/Settings/SettingsClasses.cpp
@@ -305,4 +305,21 @@ void UProjectPackagingSettings::PostEditChangeProperty( FPropertyChangedEvent& P
 		FPaths::MakePathRelativeTo(Path, FPlatformProcess::BaseDir());
 		StagingDirectory.Path = Path;
 	}
+	else if (Name == FName(TEXT("ForDistribution")) || Name == FName(TEXT("BuildConfiguration")))
+	{
+		if (ForDistribution)
+		{
+			BuildConfiguration = EProjectPackagingBuildConfigurations::PPBC_Shipping;
+		}
+	}
 }
+
+bool UProjectPackagingSettings::CanEditChange( const UProperty* InProperty ) const
+{
+	if (InProperty->GetFName() == FName(TEXT("BuildConfiguration")) && ForDistribution)
+	{
+		return false;
+	}
+
+	return Super::CanEditChange(InProperty);
+}
\ No newline at end of file
