Allow lack of UI support in SourceControl module.

This is currently needed for Linux where SourceControl is expected to work headless and thus does not normally depend on Slate. Should be solved differently.

diff --git a/Engine/Source/Editor/MainFrame/Private/MainFrameModule.cpp b/Engine/Source/Editor/MainFrame/Private/MainFrameModule.cpp
--- a/Engine/Source/Editor/MainFrame/Private/MainFrameModule.cpp
+++ b/Engine/Source/Editor/MainFrame/Private/MainFrameModule.cpp
@@ -444,6 +444,13 @@ TSharedRef<SWidget> FMainFrameModule::MakeDeveloperTools() const
 
 	bool bUseSuperSearch = FParse::Param(FCommandLine::Get(), TEXT("SuperSearch"));;
 
+	TSharedPtr<SWidget> StatusWidget = ISourceControlModule::Get().CreateStatusWidget();
+	// source control module can be compiled without UI support, provide a fallback
+	if (!StatusWidget.IsValid())
+	{
+		StatusWidget = SNullWidget::NullWidget;
+	}
+	
 	// Invisible border, so that we can animate our box panel size
 	return SNew( SBorder )
 		.Visibility( EVisibility::SelfHitTestInvisible )
@@ -504,7 +511,7 @@ TSharedRef<SWidget> FMainFrameModule::MakeDeveloperTools() const
 				.Padding( 6.0f, 0.0f, 2.0f, 0.0f )
 				[
 					// STATUS BUTTON
-					ISourceControlModule::Get().CreateStatusWidget().ToSharedRef()
+					StatusWidget.ToSharedRef()
 				]
 
 			// Editor live streaming toggle button
