Fixed a crash when trying to start up the engine without supplying a project file to get to the project browser (i.e. running the UE4 project directly from VS)

diff --git a/Engine/Source/Editor/MainFrame/Private/Frame/MainFrameHandler.h b/Engine/Source/Editor/MainFrame/Private/Frame/MainFrameHandler.h
--- a/Engine/Source/Editor/MainFrame/Private/Frame/MainFrameHandler.h
+++ b/Engine/Source/Editor/MainFrame/Private/Frame/MainFrameHandler.h
@@ -237,7 +237,10 @@ public:
 	{
 		TSharedRef<FGlobalTabmanager> GlobalTabManager = FGlobalTabmanager::Get();
 		
-		GlobalTabManager->SetMainTab(MainTab.ToSharedRef());
+		if (MainTab.IsValid())
+		{
+			GlobalTabManager->SetMainTab(MainTab.ToSharedRef());
+		}
 
 		// Persistent layouts should get stored using the specified method.
 		GlobalTabManager->SetOnPersistLayout(FTabManager::FOnPersistLayout::CreateRaw(this, &FMainFrameHandler::HandleTabManagerPersistLayout));
