Added editor notification when automatic hot reload completes

diff --git a/Engine/Source/Developer/HotReload/Private/HotReload.cpp b/Engine/Source/Developer/HotReload/Private/HotReload.cpp
--- a/Engine/Source/Developer/HotReload/Private/HotReload.cpp
+++ b/Engine/Source/Developer/HotReload/Private/HotReload.cpp
@@ -84,10 +84,14 @@ private:
 	 */
 	void OnHotReloadBinariesChanged(const TArray<struct FFileChangeData>& FileChanges);
 
-	/**	Broadcasts that a hot reload just finished. */
-	void BroadcastHotReload() 
+	/**
+	 * Broadcasts that a hot reload just finished. 
+	 *
+	 * @param	bWasTriggeredAutomatically	True if the hot reload was invoked automatically by the hot reload system after detecting a changed DLL
+	 */
+	void BroadcastHotReload( bool bWasTriggeredAutomatically ) 
 	{ 
-		HotReloadEvent.Broadcast(); 
+		HotReloadEvent.Broadcast( bWasTriggeredAutomatically ); 
 	}
 
 	/**
@@ -181,14 +185,8 @@ void FHotReloadModule::DoHotReloadFromEditor()
 
 		GetPackagesToRebindAndDependentModules(GameModuleNames, PackagesToRebind, DependentModules);
 
-		GWarn->BeginSlowTask(LOCTEXT("CompilingGameCode", "Compiling Game Code"), true);
-
 		const bool bWaitForCompletion = false;	// Don't wait -- we want compiling to happen asynchronously
 		Result = RebindPackagesInternal(PackagesToRebind, DependentModules, bWaitForCompletion, *GLog);
-
-		GWarn->EndSlowTask();
-
-		BroadcastHotReload();
 	}
 
 	RecordAnalyticsEvent(TEXT("Editor"), Result, Duration, PackagesToRebind.Num(), DependentModules.Num());
@@ -309,6 +307,10 @@ ECompilationResult::Type FHotReloadModule::DoHotReloadInternal(bool bRecompileFi
 			HotReloadFunctionRemap.Empty();
 			Result = ECompilationResult::Succeeded;
 		}
+
+
+		const bool bWasTriggeredAutomatically = !bIsHotReloadingFromEditor;
+		BroadcastHotReload( bWasTriggeredAutomatically );
 	}
 	else if (bRecompileFinished)
 	{
@@ -566,7 +568,6 @@ void FHotReloadModule::DoHotReloadFromIDE()
 		}
 
 		GWarn->EndSlowTask();
-		BroadcastHotReload();		
 	}
 
 	RecordAnalyticsEvent(TEXT("IDE"), Result, Duration, PackagesToRebind.Num(), DependentModules.Num());
diff --git a/Engine/Source/Editor/BehaviorTreeEditor/Classes/BehaviorTreeEditorTypes.h b/Engine/Source/Editor/BehaviorTreeEditor/Classes/BehaviorTreeEditorTypes.h
--- a/Engine/Source/Editor/BehaviorTreeEditor/Classes/BehaviorTreeEditorTypes.h
+++ b/Engine/Source/Editor/BehaviorTreeEditor/Classes/BehaviorTreeEditorTypes.h
@@ -86,6 +86,7 @@ struct FClassBrowseHelper
 	void OnAssetAdded(const class FAssetData& AssetData);
 	void OnAssetRemoved(const class FAssetData& AssetData);
 	void InvalidateCache();
+	void OnHotReload( bool bWasTriggeredAutomatically );
 	
 	static void AddUnknownClass(const FClassData& ClassData);
 	static bool IsClassKnown(const FClassData& ClassData);
diff --git a/Engine/Source/Editor/BehaviorTreeEditor/Private/BehaviorTreeEditorTypes.cpp b/Engine/Source/Editor/BehaviorTreeEditor/Private/BehaviorTreeEditorTypes.cpp
--- a/Engine/Source/Editor/BehaviorTreeEditor/Private/BehaviorTreeEditorTypes.cpp
+++ b/Engine/Source/Editor/BehaviorTreeEditor/Private/BehaviorTreeEditorTypes.cpp
@@ -146,7 +146,7 @@ FClassBrowseHelper::FClassBrowseHelper()
 
 	// Register to have Populate called when doing a Hot Reload.
 	IHotReloadInterface& HotReloadSupport = FModuleManager::LoadModuleChecked<IHotReloadInterface>("HotReload");
-	HotReloadSupport.OnHotReload().AddRaw(this, &FClassBrowseHelper::InvalidateCache);
+	HotReloadSupport.OnHotReload().AddRaw(this, &FClassBrowseHelper::OnHotReload);
 
 	// Register to have Populate called when a Blueprint is compiled.
 	GEditor->OnBlueprintCompiled().AddRaw( this, &FClassBrowseHelper::InvalidateCache );
@@ -307,6 +307,11 @@ void FClassBrowseHelper::InvalidateCache()
 	UpdateAvailableNodeClasses();
 }
 
+void FClassBrowseHelper::OnHotReload( bool bWasTriggeredAutomatically )
+{
+	InvalidateCache();
+}
+
 TSharedPtr<FClassDataNode> FClassBrowseHelper::CreateClassDataNode(const class FAssetData& AssetData)
 {
 	TSharedPtr<FClassDataNode> Node;
diff --git a/Engine/Source/Editor/ClassViewer/Private/SClassViewer.cpp b/Engine/Source/Editor/ClassViewer/Private/SClassViewer.cpp
--- a/Engine/Source/Editor/ClassViewer/Private/SClassViewer.cpp
+++ b/Engine/Source/Editor/ClassViewer/Private/SClassViewer.cpp
@@ -382,6 +382,9 @@ private:
 	 */	
 	void AddChildren_NoFilter( TSharedPtr< FClassViewerNode >& InOutRootNode, const TMultiMap<FName,FAssetData>& BlueprintPackageToAssetDataMap );
 
+	/** Called when hot reload has finished */
+	void OnHotReload( bool bWasTriggeredAutomatically );
+
 	/** Finds the node, recursively going deeper into the hierarchy. Does so by comparing generated class package names.
 	 *	@param InGeneratedClassPackageName			The name of the generated class package to find the node for.
 	 *
@@ -1587,7 +1590,7 @@ FClassHierarchy::FClassHierarchy()
 
 	// Register to have Populate called when doing a Hot Reload.
 	IHotReloadInterface& HotReloadSupport = FModuleManager::LoadModuleChecked<IHotReloadInterface>("HotReload");
-	HotReloadSupport.OnHotReload().AddStatic(ClassViewer::Helpers::RequestPopulateClassHierarchy);
+	HotReloadSupport.OnHotReload().AddRaw( this, &FClassHierarchy::OnHotReload );
 
 	// Register to have Populate called when a Blueprint is compiled.
 	GEditor->OnBlueprintCompiled().AddStatic(ClassViewer::Helpers::RequestPopulateClassHierarchy);
@@ -1608,7 +1611,7 @@ FClassHierarchy::~FClassHierarchy()
 
 		// Unregister to have Populate called when doing a Hot Reload.
 		IHotReloadInterface& HotReloadSupport = FModuleManager::LoadModuleChecked<IHotReloadInterface>("HotReload");
-		HotReloadSupport.OnHotReload().RemoveStatic(ClassViewer::Helpers::RequestPopulateClassHierarchy);
+		HotReloadSupport.OnHotReload().RemoveAll( this );
 
 		// Unregister to have Populate called when a Blueprint is compiled.
 		GEditor->OnBlueprintCompiled().RemoveStatic(ClassViewer::Helpers::RequestPopulateClassHierarchy);
@@ -1648,6 +1651,11 @@ static TSharedPtr< FClassViewerNode > CreateNodeForClass(UClass* Class, const TM
 	return NewNode;
 }
 
+void FClassHierarchy::OnHotReload( bool bWasTriggeredAutomatically )
+{
+	ClassViewer::Helpers::RequestPopulateClassHierarchy();
+}
+
 void FClassHierarchy::AddChildren_NoFilter( TSharedPtr< FClassViewerNode >& InOutRootNode, const TMultiMap<FName, FAssetData>& BlueprintPackageToAssetDataMap )
 {
 	UClass* RootClass = UObject::StaticClass();
diff --git a/Engine/Source/Editor/MainFrame/MainFrame.Build.cs b/Engine/Source/Editor/MainFrame/MainFrame.Build.cs
--- a/Engine/Source/Editor/MainFrame/MainFrame.Build.cs
+++ b/Engine/Source/Editor/MainFrame/MainFrame.Build.cs
@@ -74,7 +74,8 @@ public class MainFrame : ModuleRules
 				"SuperSearch",
                 "TranslationEditor",
 				"SourceCodeAccess",
-				"EditorLiveStreaming"
+				"EditorLiveStreaming",
+				"HotReload"
 			}
 		);
 	}
diff --git a/Engine/Source/Editor/MainFrame/Private/MainFrameModule.cpp b/Engine/Source/Editor/MainFrame/Private/MainFrameModule.cpp
--- a/Engine/Source/Editor/MainFrame/Private/MainFrameModule.cpp
+++ b/Engine/Source/Editor/MainFrame/Private/MainFrameModule.cpp
@@ -3,6 +3,7 @@
 #include "MainFramePrivatePCH.h"
 #include "CompilerResultsLog.h"
 #include "Editor/EditorLiveStreaming/Public/IEditorLiveStreaming.h"
+#include "Developer/HotReload/Public/IHotReload.h"
 
 DEFINE_LOG_CATEGORY(LogMainFrame);
 #define LOCTEXT_NAMESPACE "FMainFrameModule"
@@ -292,6 +293,7 @@ TSharedRef<SWidget> FMainFrameModule::MakeDeveloperTools() const
 			// Inform the user of the result of the operation
 			FNotificationInfo Info( VideoSaveResultText );
 			Info.ExpireDuration = 5.0f;
+			Info.FadeOutDuration = 0.5f;
 			Info.bUseSuccessFailIcons = false;
 			Info.bUseLargeFont = false;
 			if( HyperLinkText != "" )
@@ -598,6 +600,9 @@ void FMainFrameModule::StartupModule( )
 	FModuleManager::Get().OnModuleCompilerStarted().AddRaw( this, &FMainFrameModule::HandleLevelEditorModuleCompileStarted );
 	FModuleManager::Get().OnModuleCompilerFinished().AddRaw( this, &FMainFrameModule::HandleLevelEditorModuleCompileFinished );
 
+	// Register to find out about when hot reload completes, so we can show a notification
+	IHotReloadModule::Get().OnHotReload().AddRaw( this, &FMainFrameModule::HandleHotReloadFinished );
+
 #if WITH_EDITOR
 	ISourceCodeAccessModule& SourceCodeAccessModule = FModuleManager::LoadModuleChecked<ISourceCodeAccessModule>("SourceCodeAccess");
 	SourceCodeAccessModule.OnLaunchingCodeAccessor().AddRaw( this, &FMainFrameModule::HandleCodeAccessorLaunching );
@@ -633,6 +638,11 @@ void FMainFrameModule::ShutdownModule( )
 
 	FMainFrameCommands::Unregister();
 
+	if( IHotReloadModule::IsAvailable() )
+	{
+		IHotReloadModule::Get().OnHotReload().RemoveAll( this );
+	}
+
 	FModuleManager::Get().OnModuleCompilerStarted().RemoveAll( this );
 	FModuleManager::Get().OnModuleCompilerFinished().RemoveAll( this );
 
@@ -741,6 +751,7 @@ void FMainFrameModule::HandleLevelEditorModuleCompileFinished(const FString& Log
 		{
 			GEditor->PlayPreviewSound(CompileSuccessSound);
 			NotificationItem->SetText(NSLOCTEXT("MainFrame", "RecompileComplete", "Compile Complete!"));
+			NotificationItem->SetExpireDuration( 1.5f );
 			NotificationItem->SetCompletionState(SNotificationItem::CS_Success);
 		}
 		else
@@ -775,6 +786,31 @@ void FMainFrameModule::HandleLevelEditorModuleCompileFinished(const FString& Log
 }
 
 
+void FMainFrameModule::HandleHotReloadFinished( bool bWasTriggeredAutomatically )
+{
+	// Only play the notification for hot reloads that were triggered automatically.  If the user triggered the hot reload, they'll
+	// have a different visual cue for that, such as the "Compiling Complete!" notification
+	if( bWasTriggeredAutomatically )
+	{
+		FNotificationInfo Info( LOCTEXT("HotReloadFinished", "Hot Reload Complete!") );
+		Info.Image = FEditorStyle::GetBrush(TEXT("LevelEditor.RecompileGameCode"));
+		Info.FadeInDuration = 0.1f;
+		Info.FadeOutDuration = 0.5f;
+		Info.ExpireDuration = 1.5f;
+		Info.bUseThrobber = false;
+		Info.bUseSuccessFailIcons = true;
+		Info.bUseLargeFont = true;
+		Info.bFireAndForget = false;
+		Info.bAllowThrottleWhenFrameRateIsLow = false;
+		auto NotificationItem = FSlateNotificationManager::Get().AddNotification( Info );
+		NotificationItem->SetCompletionState(SNotificationItem::CS_Success);
+		NotificationItem->ExpireAndFadeout();
+	
+		GEditor->PlayPreviewSound(CompileSuccessSound);
+	}
+}
+
+
 void FMainFrameModule::HandleCodeAccessorLaunched( const bool WasSuccessful )
 {
 	TSharedPtr<SNotificationItem> NotificationItem = CodeAccessorNotificationPtr.Pin();
diff --git a/Engine/Source/Editor/MainFrame/Private/MainFrameModule.h b/Engine/Source/Editor/MainFrame/Private/MainFrameModule.h
--- a/Engine/Source/Editor/MainFrame/Private/MainFrameModule.h
+++ b/Engine/Source/Editor/MainFrame/Private/MainFrameModule.h
@@ -137,6 +137,9 @@ private:
 	// Handles the level editor module finishing to recompile.
 	void HandleLevelEditorModuleCompileFinished( const FString& LogDump, ECompilationResult::Type CompilationResult, bool bShowLog );
 
+	/** Called when Hot Reload completes */
+	void HandleHotReloadFinished( bool bWasTriggeredAutomatically );
+
 	// Handles the code accessor having finished launching its editor
 	void HandleCodeAccessorLaunched( const bool WasSuccessful );
 
diff --git a/Engine/Source/Runtime/CoreUObject/Public/Misc/HotReloadInterface.h b/Engine/Source/Runtime/CoreUObject/Public/Misc/HotReloadInterface.h
--- a/Engine/Source/Runtime/CoreUObject/Public/Misc/HotReloadInterface.h
+++ b/Engine/Source/Runtime/CoreUObject/Public/Misc/HotReloadInterface.h
@@ -30,8 +30,11 @@ public:
 	*/
 	virtual void RebindPackages(TArray<UPackage*> Packages, TArray<FName> DependentModules, const bool bWaitForCompletion, FOutputDevice &Ar) = 0;
 
-	/** Called when a Hot Reload event has completed. */
-	DECLARE_MULTICAST_DELEGATE(FHotReloadEvent);
+	/** Called when a Hot Reload event has completed. 
+	 * 
+	 * @param	bWasTriggeredAutomatically	True if the hot reload was invoked automatically by the hot reload system after detecting a changed DLL
+	 */
+	DECLARE_MULTICAST_DELEGATE_OneParam(FHotReloadEvent, bool /* bWasTriggeredAutomatically */ );
 	virtual FHotReloadEvent& OnHotReload() = 0;
 };
 
