Use SSpacer instead of NullWidget.

Advised by crackpots (NickA).

diff --git a/Engine/Source/Editor/MainFrame/Private/MainFrameModule.cpp b/Engine/Source/Editor/MainFrame/Private/MainFrameModule.cpp
--- a/Engine/Source/Editor/MainFrame/Private/MainFrameModule.cpp
+++ b/Engine/Source/Editor/MainFrame/Private/MainFrameModule.cpp
@@ -450,7 +450,7 @@ TSharedRef<SWidget> FMainFrameModule::MakeDeveloperTools() const
 	// source control module can be compiled without UI support, provide a fallback
 	if (!StatusWidget.IsValid())
 	{
-		StatusWidget = SNullWidget::NullWidget;
+		StatusWidget = SNew( SSpacer );
 	}
 	
 	// Invisible border, so that we can animate our box panel size
