Enable search for help by default

diff --git a/Engine/Source/Editor/MainFrame/Private/MainFrameModule.cpp b/Engine/Source/Editor/MainFrame/Private/MainFrameModule.cpp
--- a/Engine/Source/Editor/MainFrame/Private/MainFrameModule.cpp
+++ b/Engine/Source/Editor/MainFrame/Private/MainFrameModule.cpp
@@ -451,7 +451,7 @@ TSharedRef<SWidget> FMainFrameModule::MakeDeveloperTools() const
 		]
 	;
 
-	bool bUseSuperSearch = FParse::Param(FCommandLine::Get(), TEXT("SuperSearch"));;
+	bool bUseSuperSearch = true;
 
 	TSharedPtr<SWidget> StatusWidget = ISourceControlModule::Get().CreateStatusWidget();
 	// source control module can be compiled without UI support, provide a fallback
