Packaging: Added packaging prerequisite errors to MessageLog

diff --git a/Engine/Source/Editor/MainFrame/Private/Frame/MainFrameActions.cpp b/Engine/Source/Editor/MainFrame/Private/Frame/MainFrameActions.cpp
--- a/Engine/Source/Editor/MainFrame/Private/Frame/MainFrameActions.cpp
+++ b/Engine/Source/Editor/MainFrame/Private/Frame/MainFrameActions.cpp
@@ -2,11 +2,14 @@
 
 
 #include "MainFramePrivatePCH.h"
+#include "MessageLog.h"
+
 
 #define LOCTEXT_NAMESPACE "MainFrameActions"
 
 DEFINE_LOG_CATEGORY_STATIC(MainFrameActions, Log, All);
 
+
 TSharedRef< FUICommandList > FMainFrameCommands::ActionList( new FUICommandList() );
 
 TWeakPtr<SNotificationItem> FMainFrameActionCallbacks::ChoosePackagesToCheckInNotification;
@@ -21,8 +24,7 @@ FMainFrameCommands::FMainFrameCommands()
 		TEXT( "MainFrame.ToggleFullscreen" ),
 		TEXT( "Toggles the editor between \"full screen\" mode and \"normal\" mode.  In full screen mode, the task bar and window title area are hidden." ),
 		FConsoleCommandDelegate::CreateStatic( &FMainFrameActionCallbacks::ToggleFullscreen_Execute ) )
-{
-}
+{ }
 
 
 void FMainFrameCommands::RegisterCommands()
@@ -524,15 +526,56 @@ void FMainFrameActionCallbacks::PackageProject( const FName InPlatformInfoName )
 			FString NotInstalledDocLink;
 			FString ProjectPath = FPaths::IsProjectFilePathSet() ? FPaths::ConvertRelativePathToFull(FPaths::GetProjectFilePath()) : FPaths::RootDir() / FApp::GetGameName() / FApp::GetGameName() + TEXT(".uproject");
 			int32 Result = Platform->DoesntHaveRequirements(ProjectPath, bProjectHasCode, NotInstalledDocLink);
+
+			// report to analytics
 			FEditorAnalytics::ReportBuildRequirementsFailure(TEXT("Editor.Package.Failed"), PlatformInfo->TargetPlatformName.ToString(), bProjectHasCode, Result);
+
+			// report to message log
+			FText MessageLogText;
+			FText MessageLogTextDetail;
+
+			switch (Result)
+			{
+				case ETargetPlatformReadyStatus::SDKNotFound:
+					MessageLogText = LOCTEXT("SdkNotFoundMessage", "Software Development Kit (SDK) not be found");
+					MessageLogTextDetail = FText::Format(LOCTEXT("SdkNotFoundMessageDetail", "Please install the SDK for the {0} target platform!"), Platform->DisplayName());
+					break;
+
+				case ETargetPlatformReadyStatus::ProvisionNotFound:
+					MessageLogText = LOCTEXT("ProvisionNotFoundMessage", "Provision not found");
+					MessageLogTextDetail = LOCTEXT("ProvisionNotFoundMessageDetail", "A provision is required for deploying your app to the device.");
+					break;
+
+				case ETargetPlatformReadyStatus::SigningKeyNotFound:
+					MessageLogText = LOCTEXT("SigningKeyNotFoundMessage", "Signing key not found");
+					MessageLogTextDetail = LOCTEXT("SigningKeyNotFoundMessageDetail", "The app could not be digitally signed, because the signing key is not configured.");
+					break;
+
+				default:
+					break;
+			}
+
+			if (!MessageLogText.IsEmpty())
+			{
+				TSharedRef<FTokenizedMessage> Message = FTokenizedMessage::Create(EMessageSeverity::Error);
+				Message->AddToken(FTextToken::Create(MessageLogText));
+				Message->AddToken(FTextToken::Create(MessageLogTextDetail));
+				Message->AddToken(FDocumentationToken::Create(NotInstalledDocLink));
+
+				FMessageLog MessageLog("PackagingResults");
+				MessageLog.AddMessage(Message);
+				MessageLog.Open();
+			}
+
+			// report to main frame
 			switch (Result)
 			{
-				// broadcast this, and assume someone will pick it up
 			case ETargetPlatformReadyStatus::SDKNotFound:
 			case ETargetPlatformReadyStatus::ProvisionNotFound:
 			case ETargetPlatformReadyStatus::SigningKeyNotFound:
 			case (ETargetPlatformReadyStatus::ProvisionNotFound | ETargetPlatformReadyStatus::SigningKeyNotFound):
 				{
+					// broadcast this, and assume someone will pick it up
 					IMainFrameModule& MainFrameModule = FModuleManager::GetModuleChecked<IMainFrameModule>(TEXT("MainFrame"));
 					MainFrameModule.BroadcastMainFrameSDKNotInstalled(PlatformInfo->TargetPlatformName.ToString(), NotInstalledDocLink);
 				}
