Editor: Grouped Editor Preferences, Project Settings and Plugin Editor into same tab window

diff --git a/Engine/Source/Editor/MainFrame/Private/MainFrameModule.cpp b/Engine/Source/Editor/MainFrame/Private/MainFrameModule.cpp
--- a/Engine/Source/Editor/MainFrame/Private/MainFrameModule.cpp
+++ b/Engine/Source/Editor/MainFrame/Private/MainFrameModule.cpp
@@ -118,9 +118,10 @@ void FMainFrameModule::CreateDefaultMainFrame( const bool bStartImmersivePIE )
 			TSharedRef<FTabManager::FLayout> LoadedLayout = FLayoutSaveRestore::LoadFromConfig(GEditorLayoutIni,
 				// We persist the positioning of the level editor and the content browser.
 				// The asset editors currently do not get saved.
-				FTabManager::NewLayout( "UnrealEd_Layout_v1.3" )
+				FTabManager::NewLayout( "UnrealEd_Layout_v1.4" )
 				->AddArea
 				(
+					// level editor window
 					FTabManager::NewPrimaryArea()
 					->Split
 					(
@@ -132,6 +133,7 @@ void FMainFrameModule::CreateDefaultMainFrame( const bool bStartImmersivePIE )
 				)
 				->AddArea
 				(
+					// content browser window
 					FTabManager::NewArea(WindowSize)
 					->Split
 					(
@@ -142,6 +144,7 @@ void FMainFrameModule::CreateDefaultMainFrame( const bool bStartImmersivePIE )
 				)
 				->AddArea
 				(
+					// toolkits window
 					FTabManager::NewArea(WindowSize)
 					->SetOrientation(Orient_Vertical)
 					->Split
@@ -157,6 +160,19 @@ void FMainFrameModule::CreateDefaultMainFrame( const bool bStartImmersivePIE )
 						->AddTab("MergeTool", ETabState::ClosedTab)
 					)
 				)
+				->AddArea
+				(
+					// settings window
+					FTabManager::NewArea(WindowSize)
+					->Split
+					(
+						FTabManager::NewStack()
+						->SetSizeCoefficient(1.0f)
+						->AddTab("EditorSettings", ETabState::ClosedTab)
+						->AddTab("ProjectSettings", ETabState::ClosedTab)
+						->AddTab("PluginsEditor", ETabState::ClosedTab)
+					)
+				)
 			);
 
 			MainFrameContent = FGlobalTabmanager::Get()->RestoreFrom( LoadedLayout, RootWindow, bEmbedTitleAreaContent );
