Packaging: Added proper documentation link for missing SDKs

diff --git a/Engine/Source/Editor/MainFrame/Private/Frame/MainFrameActions.cpp b/Engine/Source/Editor/MainFrame/Private/Frame/MainFrameActions.cpp
--- a/Engine/Source/Editor/MainFrame/Private/Frame/MainFrameActions.cpp
+++ b/Engine/Source/Editor/MainFrame/Private/Frame/MainFrameActions.cpp
@@ -560,7 +560,7 @@ void FMainFrameActionCallbacks::PackageProject( const FName InPlatformInfoName )
 				TSharedRef<FTokenizedMessage> Message = FTokenizedMessage::Create(EMessageSeverity::Error);
 				Message->AddToken(FTextToken::Create(MessageLogText));
 				Message->AddToken(FTextToken::Create(MessageLogTextDetail));
-				Message->AddToken(FDocumentationToken::Create(NotInstalledDocLink));
+				Message->AddToken(FDocumentationToken::Create(TEXT("Platforms/iOS/QuickStart/6")));
 
 				FMessageLog MessageLog("PackagingResults");
 				MessageLog.AddMessage(Message);
