SCC: Added docs URL to SCC error message

diff --git a/Engine/Source/Editor/MainFrame/Private/Frame/MainFrameActions.cpp b/Engine/Source/Editor/MainFrame/Private/Frame/MainFrameActions.cpp
--- a/Engine/Source/Editor/MainFrame/Private/Frame/MainFrameActions.cpp
+++ b/Engine/Source/Editor/MainFrame/Private/Frame/MainFrameActions.cpp
@@ -338,7 +338,8 @@ void FMainFrameActionCallbacks::ChoosePackagesToCheckIn()
 		else
 		{
 			FMessageLog EditorErrors("EditorErrors");
-			EditorErrors.Warning(LOCTEXT( "NoSCCConnection", "No connection to source control available!"));
+			EditorErrors.Warning(LOCTEXT( "NoSCCConnection", "No connection to source control available!"))
+				->AddToken(FDocumentationToken::Create(TEXT("Engine/UI/SourceControl")));
 			EditorErrors.Notify();
 		}
 	}
