Packaging: Now using a tutorial link in the message log instead of always popping up the tutorial

diff --git a/Engine/Source/Editor/MainFrame/Private/Frame/MainFrameActions.cpp b/Engine/Source/Editor/MainFrame/Private/Frame/MainFrameActions.cpp
--- a/Engine/Source/Editor/MainFrame/Private/Frame/MainFrameActions.cpp
+++ b/Engine/Source/Editor/MainFrame/Private/Frame/MainFrameActions.cpp
@@ -561,6 +561,7 @@ void FMainFrameActionCallbacks::PackageProject( const FName InPlatformInfoName )
 				TSharedRef<FTokenizedMessage> Message = FTokenizedMessage::Create(EMessageSeverity::Error);
 				Message->AddToken(FTextToken::Create(MessageLogText));
 				Message->AddToken(FTextToken::Create(MessageLogTextDetail));
+				Message->AddToken(FTutorialToken::Create(NotInstalledTutorialLink));
 				Message->AddToken(FDocumentationToken::Create(TEXT("Platforms/iOS/QuickStart/6")));
 
 				FMessageLog MessageLog("PackagingResults");
@@ -571,17 +572,6 @@ void FMainFrameActionCallbacks::PackageProject( const FName InPlatformInfoName )
 			// report to main frame
 			switch (Result)
 			{
-			case ETargetPlatformReadyStatus::SDKNotFound:
-			case ETargetPlatformReadyStatus::ProvisionNotFound:
-			case ETargetPlatformReadyStatus::SigningKeyNotFound:
-			case (ETargetPlatformReadyStatus::ProvisionNotFound | ETargetPlatformReadyStatus::SigningKeyNotFound):
-				{
-					// broadcast this, and assume someone will pick it up
-					IMainFrameModule& MainFrameModule = FModuleManager::GetModuleChecked<IMainFrameModule>(TEXT("MainFrame"));
-					MainFrameModule.BroadcastMainFrameSDKNotInstalled(PlatformInfo->TargetPlatformName.ToString(), NotInstalledTutorialLink);
-				}
-				return;
-
 #if PLATFORM_WINDOWS
 			case ETargetPlatformReadyStatus::CodeUnsupported:
 				// show the message
