Fix mac shutdown crash

diff --git a/Engine/Source/Editor/LevelEditor/Private/LevelEditor.cpp b/Engine/Source/Editor/LevelEditor/Private/LevelEditor.cpp
--- a/Engine/Source/Editor/LevelEditor/Private/LevelEditor.cpp
+++ b/Engine/Source/Editor/LevelEditor/Private/LevelEditor.cpp
@@ -498,7 +498,7 @@ TSharedPtr<ILevelEditor> FLevelEditorModule::GetFirstLevelEditor()
 
 TSharedPtr<SDockTab> FLevelEditorModule::GetLevelEditorTab() const
 {
-	return LevelEditorInstanceTabPtr.Pin().ToSharedRef();
+	return LevelEditorInstanceTabPtr.Pin();
 }
 
 void FLevelEditorModule::BindGlobalLevelEditorCommands()
diff --git a/Engine/Source/Editor/MainFrame/Private/Frame/MainFrameActions.cpp b/Engine/Source/Editor/MainFrame/Private/Frame/MainFrameActions.cpp
--- a/Engine/Source/Editor/MainFrame/Private/Frame/MainFrameActions.cpp
+++ b/Engine/Source/Editor/MainFrame/Private/Frame/MainFrameActions.cpp
@@ -878,10 +878,10 @@ void FMainFrameActionCallbacks::ToggleFullscreen_Execute()
 
 bool FMainFrameActionCallbacks::FullScreen_IsChecked()
 {
-	const TWeakPtr<SDockTab> LevelEditorTabPtr = FModuleManager::Get().GetModuleChecked<FLevelEditorModule>( "LevelEditor" ).GetLevelEditorTab();
+	const TSharedPtr<SDockTab> LevelEditorTabPtr = FModuleManager::Get().GetModuleChecked<FLevelEditorModule>( "LevelEditor" ).GetLevelEditorTab();
 
 	const TSharedPtr<SWindow> LevelEditorWindow = LevelEditorTabPtr.IsValid()
-		? LevelEditorTabPtr.Pin()->GetParentWindow()
+		? LevelEditorTabPtr->GetParentWindow()
 		: TSharedPtr<SWindow>();
 
 	return (LevelEditorWindow.IsValid())
