On Mac OS X you can actually open the File > Connect to source control... dialog while a modal window is open, which should also be modal otherwise it won't display.

diff --git a/Engine/Source/Editor/MainFrame/Private/Frame/MainFrameActions.cpp b/Engine/Source/Editor/MainFrame/Private/Frame/MainFrameActions.cpp
--- a/Engine/Source/Editor/MainFrame/Private/Frame/MainFrameActions.cpp
+++ b/Engine/Source/Editor/MainFrame/Private/Frame/MainFrameActions.cpp
@@ -352,7 +352,8 @@ bool FMainFrameActionCallbacks::CanChoosePackagesToCheckIn()
 
 void FMainFrameActionCallbacks::ConnectToSourceControl()
 {
-	ISourceControlModule::Get().ShowLoginDialog(FSourceControlLoginClosed(), ELoginWindowMode::Modeless);
+	ELoginWindowMode::Type Mode = !FSlateApplication::Get().GetActiveModalWindow().IsValid() ? ELoginWindowMode::Modeless : ELoginWindowMode::Modal;
+	ISourceControlModule::Get().ShowLoginDialog(FSourceControlLoginClosed(), Mode);
 }
 
 void FMainFrameActionCallbacks::SaveAll()
