Editor: Added ability to respond to multiple deployment errors; also saved a kitten

diff --git a/Engine/Source/Developer/TargetPlatform/Public/Common/TargetPlatformBase.h b/Engine/Source/Developer/TargetPlatform/Public/Common/TargetPlatformBase.h
--- a/Engine/Source/Developer/TargetPlatform/Public/Common/TargetPlatformBase.h
+++ b/Engine/Source/Developer/TargetPlatform/Public/Common/TargetPlatformBase.h
@@ -166,12 +166,12 @@ public:
 		return true;
 	}
 
-	virtual int DoesntHaveRequirements(const FString& ProjectPath, bool bProjectHasCode, FString& OutDocumentationPath) const override
+	virtual int32 DoesntHaveRequirements(const FString& ProjectPath, bool bProjectHasCode, FString& OutDocumentationPath) const override
 	{
-		int bReadyToBuild = static_cast<int>(ETargetPlatformReadyStatus::Ready); // @todo How do we check that the iOS SDK is installed when building from Windows? Is that even possible?
+		int32 bReadyToBuild = ETargetPlatformReadyStatus::Ready; // @todo How do we check that the iOS SDK is installed when building from Windows? Is that even possible?
 		if (!IsSdkInstalled(bProjectHasCode, OutDocumentationPath))
 		{
-			bReadyToBuild |= static_cast<int>(ETargetPlatformReadyStatus::SDKNotFound);
+			bReadyToBuild |= ETargetPlatformReadyStatus::SDKNotFound;
 		}
 		return bReadyToBuild;
 	}
diff --git a/Engine/Source/Developer/TargetPlatform/Public/Interfaces/ITargetPlatform.h b/Engine/Source/Developer/TargetPlatform/Public/Interfaces/ITargetPlatform.h
--- a/Engine/Source/Developer/TargetPlatform/Public/Interfaces/ITargetPlatform.h
+++ b/Engine/Source/Developer/TargetPlatform/Public/Interfaces/ITargetPlatform.h
@@ -115,15 +115,14 @@ public:
 	virtual FText DisplayName( ) const = 0;
 
 	/**
-	 * Checks whether the platform's build requirements are met so that we can do things like
-	 * package for the platform
+	 * Checks whether the platform's build requirements are met so that we can do things like package for the platform.
 	 *
-	 * @param ProjectPath Path to the project
-	 * @param bProjectHasCode true if the project has code, and therefore any compilation based SDK requirements should be checked
-	 * @param OutDocumentationPath Let's the platform tell the editor a path to show some information about how to fix any problem
-	 * @return Readiness status
+	 * @param ProjectPath Path to the project.
+	 * @param bProjectHasCode true if the project has code, and therefore any compilation based SDK requirements should be checked.
+	 * @param OutDocumentationPath Let's the platform tell the editor a path to show some information about how to fix any problem.
+	 * @return Readiness status.
 	 */
-	virtual int DoesntHaveRequirements(const FString& ProjectPath, bool bProjectHasCode, FString& OutDocumentationPath) const = 0;
+	virtual int32 DoesntHaveRequirements(const FString& ProjectPath, bool bProjectHasCode, FString& OutDocumentationPath) const = 0;
 
 	/**
 	 * Returns the information about this platform
diff --git a/Engine/Source/Developer/iOS/IOSTargetPlatform/Private/IOSTargetPlatform.cpp b/Engine/Source/Developer/iOS/IOSTargetPlatform/Private/IOSTargetPlatform.cpp
--- a/Engine/Source/Developer/iOS/IOSTargetPlatform/Private/IOSTargetPlatform.cpp
+++ b/Engine/Source/Developer/iOS/IOSTargetPlatform/Private/IOSTargetPlatform.cpp
@@ -106,9 +106,9 @@ static void OnOutput(FString Message)
 	UE_LOG(LogTemp, Display, TEXT("%s\n"), *Message);
 }
 
-int FIOSTargetPlatform::DoesntHaveRequirements(const FString& ProjectPath, bool bProjectHasCode, FString& OutTutorialPath) const
+int32 FIOSTargetPlatform::DoesntHaveRequirements(const FString& ProjectPath, bool bProjectHasCode, FString& OutTutorialPath) const
 {
-	int bReadyToBuild = ETargetPlatformReadyStatus::Ready; // @todo How do we check that the iOS SDK is installed when building from Windows? Is that even possible?
+	int32 bReadyToBuild = ETargetPlatformReadyStatus::Ready; // @todo How do we check that the iOS SDK is installed when building from Windows? Is that even possible?
 	if (!IsSdkInstalled(bProjectHasCode, OutTutorialPath))
 	{
 		bReadyToBuild |= ETargetPlatformReadyStatus::SDKNotFound;
diff --git a/Engine/Source/Developer/iOS/IOSTargetPlatform/Private/IOSTargetPlatform.h b/Engine/Source/Developer/iOS/IOSTargetPlatform/Private/IOSTargetPlatform.h
--- a/Engine/Source/Developer/iOS/IOSTargetPlatform/Private/IOSTargetPlatform.h
+++ b/Engine/Source/Developer/iOS/IOSTargetPlatform/Private/IOSTargetPlatform.h
@@ -84,7 +84,7 @@ public:
 	}
 
 	virtual bool IsSdkInstalled(bool bProjectHasCode, FString& OutTutorialPath) const override;
-	virtual int DoesntHaveRequirements(const FString& ProjectPath, bool bProjectHasCode, FString& OutTutorialPath) const override;
+	virtual int32 DoesntHaveRequirements(const FString& ProjectPath, bool bProjectHasCode, FString& OutTutorialPath) const override;
 
 
 #if WITH_ENGINE
diff --git a/Engine/Source/Editor/MainFrame/Private/Frame/MainFrameActions.cpp b/Engine/Source/Editor/MainFrame/Private/Frame/MainFrameActions.cpp
--- a/Engine/Source/Editor/MainFrame/Private/Frame/MainFrameActions.cpp
+++ b/Engine/Source/Editor/MainFrame/Private/Frame/MainFrameActions.cpp
@@ -533,59 +533,53 @@ void FMainFrameActionCallbacks::PackageProject( const FName InPlatformInfoName )
 			FEditorAnalytics::ReportBuildRequirementsFailure(TEXT("Editor.Package.Failed"), PlatformInfo->TargetPlatformName.ToString(), bProjectHasCode, Result);
 
 			// report to message log
-			FText MessageLogText;
-			FText MessageLogTextDetail;
+			if ((Result & ETargetPlatformReadyStatus::SDKNotFound) != 0)
+			{
+				AddMessageLog(
+					LOCTEXT("SdkNotFoundMessage", "Software Development Kit (SDK) not found."),
+					FText::Format(LOCTEXT("SdkNotFoundMessageDetail", "Please install the SDK for the {0} target platform!"), Platform->DisplayName()),
+					NotInstalledTutorialLink
+				);
+			}
 
-			switch (Result)
+			if ((Result & ETargetPlatformReadyStatus::ProvisionNotFound) != 0)
 			{
-				case ETargetPlatformReadyStatus::SDKNotFound:
-					MessageLogText = LOCTEXT("SdkNotFoundMessage", "Software Development Kit (SDK) not found.");
-					MessageLogTextDetail = FText::Format(LOCTEXT("SdkNotFoundMessageDetail", "Please install the SDK for the {0} target platform!"), Platform->DisplayName());
-					break;
-
-				case ETargetPlatformReadyStatus::ProvisionNotFound:
-					MessageLogText = LOCTEXT("ProvisionNotFoundMessage", "Provision not found.");
-					MessageLogTextDetail = LOCTEXT("ProvisionNotFoundMessageDetail", "A provision is required for deploying your app to the device.");
-					break;
-
-				case ETargetPlatformReadyStatus::SigningKeyNotFound:
-					MessageLogText = LOCTEXT("SigningKeyNotFoundMessage", "Signing key not found.");
-					MessageLogTextDetail = LOCTEXT("SigningKeyNotFoundMessageDetail", "The app could not be digitally signed, because the signing key is not configured.");
-					break;
-
-				default:
-					break;
+				AddMessageLog(
+					LOCTEXT("ProvisionNotFoundMessage", "Provision not found."),
+					LOCTEXT("ProvisionNotFoundMessageDetail", "A provision is required for deploying your app to the device."),
+					NotInstalledTutorialLink
+				);
 			}
 
-			if (!MessageLogText.IsEmpty())
+			if ((Result & ETargetPlatformReadyStatus::SigningKeyNotFound) != 0)
 			{
-				TSharedRef<FTokenizedMessage> Message = FTokenizedMessage::Create(EMessageSeverity::Error);
-				Message->AddToken(FTextToken::Create(MessageLogText));
-				Message->AddToken(FTextToken::Create(MessageLogTextDetail));
-				Message->AddToken(FTutorialToken::Create(NotInstalledTutorialLink));
-				Message->AddToken(FDocumentationToken::Create(TEXT("Platforms/iOS/QuickStart/6")));
-
-				FMessageLog MessageLog("PackagingResults");
-				MessageLog.AddMessage(Message);
-				MessageLog.Open();
+				AddMessageLog(
+					LOCTEXT("SigningKeyNotFoundMessage", "Signing key not found."),
+					LOCTEXT("SigningKeyNotFoundMessageDetail", "The app could not be digitally signed, because the signing key is not configured."),
+					NotInstalledTutorialLink
+				);
 			}
 
 			// report to main frame
-			switch (Result)
-			{
+			bool UnrecoverableError = false;
+
 #if PLATFORM_WINDOWS
-			case ETargetPlatformReadyStatus::CodeUnsupported:
-				// show the message
+			if ((Result & ETargetPlatformReadyStatus::CodeUnsupported) != 0)
+			{
 				FMessageDialog::Open(EAppMsgType::Ok, LOCTEXT("IOSNotSupported", "Sorry, launching on a device is currently not supported for code-based iOS projects. This feature will be available in a future release.") );
-				return;
+				UnrecoverableError = true;
+			}
 
-			case ETargetPlatformReadyStatus::PluginsUnsupported:
-				// show the message
+			if ((Result & ETargetPlatformReadyStatus::CodeUnsupported) != 0)
+			{
 				FMessageDialog::Open(EAppMsgType::Ok, LOCTEXT("IOSNotSupported", "Sorry, launching on a device is currently not supported for content based projects with third-party plugins. This feature will be available in a future release.") );
-				return;
+				UnrecoverableError = true;
+			}
 #endif
-			default:
-				break;
+
+			if (UnrecoverableError)
+			{
+				return;
 			}
 		}
 	}
@@ -1077,6 +1071,20 @@ void FMainFrameActionCallbacks::OpenWidgetReflector_Execute()
 /* FMainFrameActionCallbacks implementation
  *****************************************************************************/
 
+void FMainFrameActionCallbacks::AddMessageLog( const FText& Text, const FText& Detail, const FString& TutorialLink )
+{
+	TSharedRef<FTokenizedMessage> Message = FTokenizedMessage::Create(EMessageSeverity::Error);
+	Message->AddToken(FTextToken::Create(Text));
+	Message->AddToken(FTextToken::Create(Detail));
+	Message->AddToken(FTutorialToken::Create(TutorialLink));
+	Message->AddToken(FDocumentationToken::Create(TEXT("Platforms/iOS/QuickStart/6")));
+
+	FMessageLog MessageLog("PackagingResults");
+	MessageLog.AddMessage(Message);
+	MessageLog.Open();
+}
+
+
 void FMainFrameActionCallbacks::CreateUatTask( const FString& CommandLine, const FText& PlatformDisplayName, const FText& TaskName, const FText &TaskShortName, const FSlateBrush* TaskIcon )
 {
 	// make sure that the UAT batch file is in place
diff --git a/Engine/Source/Editor/MainFrame/Private/Frame/MainFrameActions.h b/Engine/Source/Editor/MainFrame/Private/Frame/MainFrameActions.h
--- a/Engine/Source/Editor/MainFrame/Private/Frame/MainFrameActions.h
+++ b/Engine/Source/Editor/MainFrame/Private/Frame/MainFrameActions.h
@@ -229,9 +229,17 @@ public:
 	// List of projects that the user can switch to.
 	static TArray<FString> ProjectNames;
 
-
 protected:
 
+	/**
+	 * Adds a message to the message log.
+	 *
+	 * @param Text The main message text.
+	 * @param Detail The detailed description.
+	 * @param TutorialLink A link to an associated tutorial.
+	 */
+	static void AddMessageLog( const FText& Text, const FText& Detail, const FString& TutorialLink );
+
 	/**
 	 * Creates an asynchronous UAT task.
 	 *
diff --git a/Engine/Source/Editor/UnrealEd/Private/Kismet2/DebuggerCommands.cpp b/Engine/Source/Editor/UnrealEd/Private/Kismet2/DebuggerCommands.cpp
--- a/Engine/Source/Editor/UnrealEd/Private/Kismet2/DebuggerCommands.cpp
+++ b/Engine/Source/Editor/UnrealEd/Private/Kismet2/DebuggerCommands.cpp
@@ -112,6 +112,15 @@ public:
 
 protected:
 
+	/**
+	 * Adds a message to the message log.
+	 *
+	 * @param Text The main message text.
+	 * @param Detail The detailed description.
+	 * @param TutorialLink A link to an associated tutorial.
+	 */
+	static void AddMessageLog( const FText& Text, const FText& Detail, const FString& TutorialLink );
+
 	/**
 	 * Checks whether the specified platform has a default device that can be launched on.
 	 *
@@ -1408,42 +1417,43 @@ bool FInternalPlayWorldCommandCallbacks::IsReadyToLaunchOnDevice(FString DeviceI
 		FEditorAnalytics::ReportBuildRequirementsFailure(TEXT("Editor.LaunchOn.Failed"), PlatformName, bHasCode, Result);
 
 		// report to message log
-		FText MessageLogText;
-		FText MessageLogTextDetail;
+		bool UnrecoverableError = false;
 
-		switch (Result)
+		if ((Result & ETargetPlatformReadyStatus::SDKNotFound) != 0)
 		{
-			case ETargetPlatformReadyStatus::SDKNotFound:
-				MessageLogText = LOCTEXT("SdkNotFoundMessage", "Software Development Kit (SDK) not found.");
-				MessageLogTextDetail = FText::Format(LOCTEXT("SdkNotFoundMessageDetail", "Please install the SDK for the {0} target platform!"), Platform->DisplayName());
-				break;
+			AddMessageLog(
+				LOCTEXT("SdkNotFoundMessage", "Software Development Kit (SDK) not found."),
+				FText::Format(LOCTEXT("SdkNotFoundMessageDetail", "Please install the SDK for the {0} target platform!"), Platform->DisplayName()),
+				NotInstalledTutorialLink
+			);
 
-			case ETargetPlatformReadyStatus::ProvisionNotFound:
-				MessageLogText = LOCTEXT("ProvisionNotFoundMessage", "Provision not found.");
-				MessageLogTextDetail = LOCTEXT("ProvisionNotFoundMessageDetail", "A provision is required for deploying your app to the device.");
-				break;
+			UnrecoverableError = true;
+		}
 
-			case ETargetPlatformReadyStatus::SigningKeyNotFound:
-				MessageLogText = LOCTEXT("SigningKeyNotFoundMessage", "Signing key not found.");
-				MessageLogTextDetail = LOCTEXT("SigningKeyNotFoundMessageDetail", "The app could not be digitally signed, because the signing key is not configured.");
-				break;
+		if ((Result & ETargetPlatformReadyStatus::ProvisionNotFound) != 0)
+		{
+			AddMessageLog(
+				LOCTEXT("ProvisionNotFoundMessage", "Provision not found."),
+				LOCTEXT("ProvisionNotFoundMessageDetail", "A provision is required for deploying your app to the device."),
+				NotInstalledTutorialLink
+			);
 
-			default:
-				break;
+			UnrecoverableError = true;
 		}
 
-		if (!MessageLogText.IsEmpty())
+		if ((Result & ETargetPlatformReadyStatus::SigningKeyNotFound) != 0)
 		{
-			TSharedRef<FTokenizedMessage> Message = FTokenizedMessage::Create(EMessageSeverity::Error);
-			Message->AddToken(FTextToken::Create(MessageLogText));
-			Message->AddToken(FTextToken::Create(MessageLogTextDetail));
-			Message->AddToken(FTutorialToken::Create(NotInstalledTutorialLink));
-			Message->AddToken(FDocumentationToken::Create(TEXT("Platforms/iOS/QuickStart/6")));
+			AddMessageLog(
+				LOCTEXT("SigningKeyNotFoundMessage", "Signing key not found."),
+				LOCTEXT("SigningKeyNotFoundMessageDetail", "The app could not be digitally signed, because the signing key is not configured."),
+				NotInstalledTutorialLink
+			);
 
-			FMessageLog MessageLog("PackagingResults");
-			MessageLog.AddMessage(Message);
-			MessageLog.Open();
+			UnrecoverableError = true;
+		}
 
+		if (UnrecoverableError)
+		{
 			return false;
 		}
 
@@ -1705,6 +1715,20 @@ bool FInternalPlayWorldCommandCallbacks::CanPossessEjectPlayer()
 }
 
 
+void FInternalPlayWorldCommandCallbacks::AddMessageLog( const FText& Text, const FText& Detail, const FString& TutorialLink )
+{
+	TSharedRef<FTokenizedMessage> Message = FTokenizedMessage::Create(EMessageSeverity::Error);
+	Message->AddToken(FTextToken::Create(Text));
+	Message->AddToken(FTextToken::Create(Detail));
+	Message->AddToken(FTutorialToken::Create(TutorialLink));
+	Message->AddToken(FDocumentationToken::Create(TEXT("Platforms/iOS/QuickStart/6")));
+
+	FMessageLog MessageLog("PackagingResults");
+	MessageLog.AddMessage(Message);
+	MessageLog.Open();
+}
+
+
 bool FInternalPlayWorldCommandCallbacks::CanLaunchOnDevice(const FString& DeviceName)
 {
 	if (!GUnrealEd->IsPlayingViaLauncher())
