Saved another kitten!

diff --git a/Engine/Source/Developer/TargetPlatform/Public/Common/TargetPlatformBase.h b/Engine/Source/Developer/TargetPlatform/Public/Common/TargetPlatformBase.h
--- a/Engine/Source/Developer/TargetPlatform/Public/Common/TargetPlatformBase.h
+++ b/Engine/Source/Developer/TargetPlatform/Public/Common/TargetPlatformBase.h
@@ -166,7 +166,7 @@ public:
 		return true;
 	}
 
-	virtual int32 DoesntHaveRequirements(const FString& ProjectPath, bool bProjectHasCode, FString& OutDocumentationPath) const override
+	virtual int32 CheckRequirements(const FString& ProjectPath, bool bProjectHasCode, FString& OutDocumentationPath) const override
 	{
 		int32 bReadyToBuild = ETargetPlatformReadyStatus::Ready; // @todo How do we check that the iOS SDK is installed when building from Windows? Is that even possible?
 		if (!IsSdkInstalled(bProjectHasCode, OutDocumentationPath))
diff --git a/Engine/Source/Developer/TargetPlatform/Public/Interfaces/ITargetPlatform.h b/Engine/Source/Developer/TargetPlatform/Public/Interfaces/ITargetPlatform.h
--- a/Engine/Source/Developer/TargetPlatform/Public/Interfaces/ITargetPlatform.h
+++ b/Engine/Source/Developer/TargetPlatform/Public/Interfaces/ITargetPlatform.h
@@ -122,7 +122,7 @@ public:
 	 * @param OutDocumentationPath Let's the platform tell the editor a path to show some information about how to fix any problem.
 	 * @return Readiness status.
 	 */
-	virtual int32 DoesntHaveRequirements(const FString& ProjectPath, bool bProjectHasCode, FString& OutDocumentationPath) const = 0;
+	virtual int32 CheckRequirements(const FString& ProjectPath, bool bProjectHasCode, FString& OutDocumentationPath) const = 0;
 
 	/**
 	 * Returns the information about this platform
diff --git a/Engine/Source/Developer/iOS/IOSTargetPlatform/Private/IOSTargetPlatform.cpp b/Engine/Source/Developer/iOS/IOSTargetPlatform/Private/IOSTargetPlatform.cpp
--- a/Engine/Source/Developer/iOS/IOSTargetPlatform/Private/IOSTargetPlatform.cpp
+++ b/Engine/Source/Developer/iOS/IOSTargetPlatform/Private/IOSTargetPlatform.cpp
@@ -106,7 +106,7 @@ static void OnOutput(FString Message)
 	UE_LOG(LogTemp, Display, TEXT("%s\n"), *Message);
 }
 
-int32 FIOSTargetPlatform::DoesntHaveRequirements(const FString& ProjectPath, bool bProjectHasCode, FString& OutTutorialPath) const
+int32 FIOSTargetPlatform::CheckRequirements(const FString& ProjectPath, bool bProjectHasCode, FString& OutTutorialPath) const
 {
 	int32 bReadyToBuild = ETargetPlatformReadyStatus::Ready; // @todo How do we check that the iOS SDK is installed when building from Windows? Is that even possible?
 	if (!IsSdkInstalled(bProjectHasCode, OutTutorialPath))
diff --git a/Engine/Source/Developer/iOS/IOSTargetPlatform/Private/IOSTargetPlatform.h b/Engine/Source/Developer/iOS/IOSTargetPlatform/Private/IOSTargetPlatform.h
--- a/Engine/Source/Developer/iOS/IOSTargetPlatform/Private/IOSTargetPlatform.h
+++ b/Engine/Source/Developer/iOS/IOSTargetPlatform/Private/IOSTargetPlatform.h
@@ -84,7 +84,7 @@ public:
 	}
 
 	virtual bool IsSdkInstalled(bool bProjectHasCode, FString& OutTutorialPath) const override;
-	virtual int32 DoesntHaveRequirements(const FString& ProjectPath, bool bProjectHasCode, FString& OutTutorialPath) const override;
+	virtual int32 CheckRequirements(const FString& ProjectPath, bool bProjectHasCode, FString& OutTutorialPath) const override;
 
 
 #if WITH_ENGINE
diff --git a/Engine/Source/Editor/MainFrame/Private/Frame/MainFrameActions.cpp b/Engine/Source/Editor/MainFrame/Private/Frame/MainFrameActions.cpp
--- a/Engine/Source/Editor/MainFrame/Private/Frame/MainFrameActions.cpp
+++ b/Engine/Source/Editor/MainFrame/Private/Frame/MainFrameActions.cpp
@@ -527,7 +527,7 @@ void FMainFrameActionCallbacks::PackageProject( const FName InPlatformInfoName )
 		{
 			FString NotInstalledTutorialLink;
 			FString ProjectPath = FPaths::IsProjectFilePathSet() ? FPaths::ConvertRelativePathToFull(FPaths::GetProjectFilePath()) : FPaths::RootDir() / FApp::GetGameName() / FApp::GetGameName() + TEXT(".uproject");
-			int32 Result = Platform->DoesntHaveRequirements(ProjectPath, bProjectHasCode, NotInstalledTutorialLink);
+			int32 Result = Platform->CheckRequirements(ProjectPath, bProjectHasCode, NotInstalledTutorialLink);
 
 			// report to analytics
 			FEditorAnalytics::ReportBuildRequirementsFailure(TEXT("Editor.Package.Failed"), PlatformInfo->TargetPlatformName.ToString(), bProjectHasCode, Result);
diff --git a/Engine/Source/Editor/UnrealEd/Private/Kismet2/DebuggerCommands.cpp b/Engine/Source/Editor/UnrealEd/Private/Kismet2/DebuggerCommands.cpp
--- a/Engine/Source/Editor/UnrealEd/Private/Kismet2/DebuggerCommands.cpp
+++ b/Engine/Source/Editor/UnrealEd/Private/Kismet2/DebuggerCommands.cpp
@@ -1411,7 +1411,7 @@ bool FInternalPlayWorldCommandCallbacks::IsReadyToLaunchOnDevice(FString DeviceI
 	{
 		FString NotInstalledTutorialLink;
 		FString ProjectPath = FPaths::IsProjectFilePathSet() ? FPaths::ConvertRelativePathToFull(FPaths::GetProjectFilePath()) : FPaths::RootDir() / FApp::GetGameName() / FApp::GetGameName() + TEXT(".uproject");
-		int32 Result = Platform->DoesntHaveRequirements(ProjectPath, bHasCode, NotInstalledTutorialLink);
+		int32 Result = Platform->CheckRequirements(ProjectPath, bHasCode, NotInstalledTutorialLink);
 		
 		// report to analytics
 		FEditorAnalytics::ReportBuildRequirementsFailure(TEXT("Editor.LaunchOn.Failed"), PlatformName, bHasCode, Result);
