Fixed a Mac-specific editor crash caused by pressing a modifier key while splash screen is on screen and some parts of the editor are already initialized, but GEditor is still NULL

diff --git a/Engine/Source/Editor/MainFrame/Private/Frame/MainFrameActions.cpp b/Engine/Source/Editor/MainFrame/Private/Frame/MainFrameActions.cpp
--- a/Engine/Source/Editor/MainFrame/Private/Frame/MainFrameActions.cpp
+++ b/Engine/Source/Editor/MainFrame/Private/Frame/MainFrameActions.cpp
@@ -163,7 +163,7 @@ FReply FMainFrameActionCallbacks::OnUnhandledKeyDownEvent(const FKeyboardEvent&
 	{
 		return FReply::Handled();
 	}
-	else if( GEditor->PlayWorld && InKeyboardEvent.GetKey() == EKeys::Escape )
+	else if( GEditor && GEditor->PlayWorld && InKeyboardEvent.GetKey() == EKeys::Escape )
 	{
 		FLevelEditorModule& LevelEditor = FModuleManager::GetModuleChecked< FLevelEditorModule >(TEXT("LevelEditor"));
 		if( LevelEditor.GetLevelEditorTab()->IsForeground()  )
