Sorted Settings submenus when "Expand configuration menus" is checked.

diff --git a/Engine/Source/Editor/MainFrame/Private/Menus/SettingsMenu.h b/Engine/Source/Editor/MainFrame/Private/Menus/SettingsMenu.h
--- a/Engine/Source/Editor/MainFrame/Private/Menus/SettingsMenu.h
+++ b/Engine/Source/Editor/MainFrame/Private/Menus/SettingsMenu.h
@@ -57,6 +57,12 @@ public:
 			{
 				MenuBuilder.BeginSection(SettingsCategory->GetName(), SettingsCategory->GetDisplayName());
 
+				SettingsSections.Sort(
+					[](const ISettingsSectionPtr& First, const ISettingsSectionPtr& Second)
+					{
+						return First->GetDisplayName().ToString() < Second->GetDisplayName().ToString();
+					});
+
 				for (int32 SectionIndex = 0; SectionIndex < SettingsSections.Num(); ++SectionIndex)
 				{
 					ISettingsSectionPtr SettingsSection = SettingsSections[SectionIndex];
