fixed issue with double message when trying to package code based IOS project on PC
UE-5158
#ios

diff --git a/Engine/Source/Editor/MainFrame/Private/Frame/MainFrameActions.cpp b/Engine/Source/Editor/MainFrame/Private/Frame/MainFrameActions.cpp
--- a/Engine/Source/Editor/MainFrame/Private/Frame/MainFrameActions.cpp
+++ b/Engine/Source/Editor/MainFrame/Private/Frame/MainFrameActions.cpp
@@ -593,8 +593,7 @@ void FMainFrameActionCallbacks::PackageProject( const FName InPlatformInfoName )
 				FMessageDialog::Open(EAppMsgType::Ok, LOCTEXT("IOSNotSupported", "Sorry, launching on a device is currently not supported for code-based iOS projects. This feature will be available in a future release.") );
 				UnrecoverableError = true;
 			}
-
-			if ((Result & ETargetPlatformReadyStatus::CodeUnsupported) != 0)
+			else if ((Result & ETargetPlatformReadyStatus::PluginsUnsupported) != 0)
 			{
 				FMessageDialog::Open(EAppMsgType::Ok, LOCTEXT("IOSNotSupported", "Sorry, launching on a device is currently not supported for content based projects with third-party plugins. This feature will be available in a future release.") );
 				UnrecoverableError = true;
