Escape no longer exits PIE unless the viewport has focus (don't want hard-coded input behavior) (integrate)

diff --git a/Engine/Source/Editor/MainFrame/Private/Frame/MainFrameActions.cpp b/Engine/Source/Editor/MainFrame/Private/Frame/MainFrameActions.cpp
--- a/Engine/Source/Editor/MainFrame/Private/Frame/MainFrameActions.cpp
+++ b/Engine/Source/Editor/MainFrame/Private/Frame/MainFrameActions.cpp
@@ -163,16 +163,6 @@ FReply FMainFrameActionCallbacks::OnUnhandledKeyDownEvent(const FKeyEvent& InKey
 	{
 		return FReply::Handled();
 	}
-	else if( GEditor && GEditor->PlayWorld && InKeyEvent.GetKey() == EKeys::Escape )
-	{
-		FLevelEditorModule& LevelEditor = FModuleManager::GetModuleChecked< FLevelEditorModule >(TEXT("LevelEditor"));
-		if( LevelEditor.GetLevelEditorTab()->IsForeground()  )
-		{
-			// Escape that gets this far should end the play map if we have one
-			// We have to do this here because nothing else has focus and escape should always end the active PIE session.
-			GEditor->RequestEndPlayMap();
-		}
-	}
 	return FReply::Unhandled();
 }
 
