Prevented the crash reporter being packaged in external code-based projects

The crash reporter used to get packaged (via the editor) for all projects.

This change makes it only get packaged for blueprint or internal projects. External code-based projects won't have the crash reporter packaged, and it can also be excluded on a per-project basis via the IncludeCrashReporter project packaging setting.

diff --git a/Engine/Config/BaseGame.ini b/Engine/Config/BaseGame.ini
--- a/Engine/Config/BaseGame.ini
+++ b/Engine/Config/BaseGame.ini
@@ -39,6 +39,7 @@ UsePakFile=True
 bGenerateChunks=False
 IncludePrerequisites=True
 UseOBB_InAPK=False
+IncludeCrashReporter=True
 InternationalizationPreset=English
 +CulturesToStage=en
 
diff --git a/Engine/Source/Editor/MainFrame/Private/Frame/MainFrameActions.cpp b/Engine/Source/Editor/MainFrame/Private/Frame/MainFrameActions.cpp
--- a/Engine/Source/Editor/MainFrame/Private/Frame/MainFrameActions.cpp
+++ b/Engine/Source/Editor/MainFrame/Private/Frame/MainFrameActions.cpp
@@ -7,6 +7,7 @@
 #include "SNotificationList.h"
 #include "NotificationManager.h"
 #include "GenericCommands.h"
+#include "EngineBuildSettings.h"
 
 
 #define LOCTEXT_NAMESPACE "MainFrameActions"
@@ -685,6 +686,12 @@ void FMainFrameActionCallbacks::PackageProject( const FName InPlatformInfoName )
 		OptionalParams += TEXT(" -build");
 	}
 
+	// we can include the crash reporter for any non-code projects or internal builds (we don't have the symbols to interpret external builds of code-based projects)
+	if (PackagingSettings->IncludeCrashReporter && (!bProjectHasCode || FEngineBuildSettings::IsInternalBuild()))
+	{
+		OptionalParams += TEXT(" -CrashReporter");
+	}
+
 	FString ExecutableName = FPlatformProcess::ExecutableName(false);
 #if PLATFORM_WINDOWS
 	// turn UE4editor into UE4editor.cmd
@@ -709,7 +716,7 @@ void FMainFrameActionCallbacks::PackageProject( const FName InPlatformInfoName )
 	Configuration = Configuration.Replace(TEXT("PPBC_"), TEXT(""));
 
 	FString ProjectPath = FPaths::IsProjectFilePathSet() ? FPaths::ConvertRelativePathToFull(FPaths::GetProjectFilePath()) : FPaths::RootDir() / FApp::GetGameName() / FApp::GetGameName() + TEXT(".uproject");
-	FString CommandLine = FString::Printf(TEXT("BuildCookRun %s%s -nop4 -project=\"%s\" -cook -allmaps -CrashReporter -stage -archive -archivedirectory=\"%s\" -package -%s -clientconfig=%s -ue4exe=%s %s -utf8output"),
+	FString CommandLine = FString::Printf(TEXT("BuildCookRun %s%s -nop4 -project=\"%s\" -cook -allmaps -stage -archive -archivedirectory=\"%s\" -package -%s -clientconfig=%s -ue4exe=%s %s -utf8output"),
 		FRocketSupport::IsRocket() ? TEXT( "-rocket -nocompile" ) : TEXT( "-nocompileeditor" ),
 		FApp::IsEngineInstalled() ? TEXT(" -installed") : TEXT(""),
 		*ProjectPath,
diff --git a/Engine/Source/Editor/UnrealEd/Classes/Settings/ProjectPackagingSettings.h b/Engine/Source/Editor/UnrealEd/Classes/Settings/ProjectPackagingSettings.h
--- a/Engine/Source/Editor/UnrealEd/Classes/Settings/ProjectPackagingSettings.h
+++ b/Engine/Source/Editor/UnrealEd/Classes/Settings/ProjectPackagingSettings.h
@@ -97,6 +97,13 @@ public:
 	UPROPERTY(config, EditAnywhere, Category = Packaging, meta = (DisplayName = "Package .pak in APK"))
 	bool UseOBB_InAPK;
 
+	/**
+	 * Specifies whether to include the crash reporter in the packaged project. 
+	 * This is included by default for Blueprint based projects, but can optionally be disabled.
+	 */
+	UPROPERTY(config, EditAnywhere, Category=Packaging, AdvancedDisplay)
+	bool IncludeCrashReporter;
+
 	/** Predefined sets of culture whose internationalization data should be packaged. */
 	UPROPERTY(config, EditAnywhere, Category=Packaging, AdvancedDisplay, meta=(DisplayName="Internationalization Support"))
 	TEnumAsByte<EProjectPackagingInternationalizationPresets> InternationalizationPreset;
diff --git a/Engine/Source/Programs/AutomationTool/Scripts/CopyBuildToStagingDirectory.Automation.cs b/Engine/Source/Programs/AutomationTool/Scripts/CopyBuildToStagingDirectory.Automation.cs
--- a/Engine/Source/Programs/AutomationTool/Scripts/CopyBuildToStagingDirectory.Automation.cs
+++ b/Engine/Source/Programs/AutomationTool/Scripts/CopyBuildToStagingDirectory.Automation.cs
@@ -1064,7 +1064,7 @@ public partial class Project : CommandUtils
 				ExecutablesToStage,
 				InDedicatedServer,
 				Params.Cook || Params.CookOnTheFly,
-				Params.CrashReporter && (StagePlatform != UnrealTargetPlatform.Linux || !Params.Rocket),
+				Params.CrashReporter && !(StagePlatform == UnrealTargetPlatform.Linux && Params.Rocket), // can't include the crash reporter from binary Linux builds
 				Params.Stage,
 				Params.CookOnTheFly,
 				Params.Archive,
