Project Launcher is no longer experimental.

diff --git a/Engine/Source/Editor/MainFrame/Private/Menus/MainMenu.cpp b/Engine/Source/Editor/MainFrame/Private/Menus/MainMenu.cpp
--- a/Engine/Source/Editor/MainFrame/Private/Menus/MainMenu.cpp
+++ b/Engine/Source/Editor/MainFrame/Private/Menus/MainMenu.cpp
@@ -166,30 +166,35 @@ void FMainMenu::FillWindowMenu( FMenuBuilder& MenuBuilder, const TSharedRef< FEx
 		TabManager->PopulateTabSpawnerMenu(MenuBuilder, MenuStructure.GetStructureRoot());
 	}
 
+	MenuBuilder.BeginSection("WindowGlobalTabSpawners");
+	{
+		MenuBuilder.AddMenuEntry(
+			LOCTEXT("ProjectLauncherLabel", "Project Launcher"),
+			LOCTEXT("ProjectLauncherToolTip", "The Project Launcher provides advanced workflows for packaging, deploying and launching your projects."),
+			FSlateIcon(FEditorStyle::GetStyleSetName(), "Launcher.TabIcon"),
+			FUIAction(FExecuteAction::CreateStatic(&FMainMenu::OpenProjectLauncher))
+			);
+
+		//@todo The tab system needs to be able to be extendable by plugins [9/3/2013 Justin.Sargent]
+		if (IModularFeatures::Get().IsModularFeatureAvailable(EditorFeatures::PluginsEditor))
+		{
+			FGlobalTabmanager::Get()->PopulateTabSpawnerMenu(MenuBuilder, "PluginsEditor");
+		}
+	}
+	MenuBuilder.EndSection();
+
 	{
 		// This is a temporary home for the spawners of experimental features that must be explicitly enabled.
 		// When the feature becomes permanent and need not check a flag, register a nomad spawner for it in the proper WorkspaceMenu category
-		bool bProjectLauncher = GetDefault<UEditorExperimentalSettings>()->bProjectLauncher;
 		bool bMessagingDebugger = GetDefault<UEditorExperimentalSettings>()->bMessagingDebugger;
 		bool bBlutility = GetDefault<UEditorExperimentalSettings>()->bEnableEditorUtilityBlueprints;
 		bool bTranslationEditor = GetDefault<UEditorExperimentalSettings>()->bEnableTranslationEditor;
 
 		// Make sure at least one is enabled before creating the section
-		if (bProjectLauncher || bMessagingDebugger || bBlutility || bTranslationEditor)
+		if (bMessagingDebugger || bBlutility || bTranslationEditor)
 		{
 			MenuBuilder.BeginSection("ExperimentalTabSpawners", LOCTEXT("ExperimentalTabSpawnersHeading", "Experimental"));
 			{
-				// Project Launcher
-				if (bProjectLauncher)
-				{
-					MenuBuilder.AddMenuEntry(
-						LOCTEXT("ProjectLauncherLabel", "Project Launcher"),
-						LOCTEXT("ProjectLauncherToolTip", "The Project Launcher provides advanced workflows for packaging, deploying and launching your projects."),
-						FSlateIcon(FEditorStyle::GetStyleSetName(), "Launcher.TabIcon"),
-						FUIAction(FExecuteAction::CreateStatic(&FMainMenu::OpenProjectLauncher))
-						);
-				}
-
 				// Messaging Debugger
 				if (bMessagingDebugger)
 				{
@@ -226,16 +231,6 @@ void FMainMenu::FillWindowMenu( FMenuBuilder& MenuBuilder, const TSharedRef< FEx
 		}
 	}
 
-	MenuBuilder.BeginSection("WindowGlobalTabSpawners");
-	{
-		//@todo The tab system needs to be able to be extendable by plugins [9/3/2013 Justin.Sargent]
-		if (IModularFeatures::Get().IsModularFeatureAvailable(EditorFeatures::PluginsEditor ) )
-		{
-			FGlobalTabmanager::Get()->PopulateTabSpawnerMenu(MenuBuilder, "PluginsEditor");
-		}
-	}
-	MenuBuilder.EndSection();
-
 	MenuBuilder.BeginSection("WindowLayout", NSLOCTEXT("MainAppMenu", "LayoutManagementHeader", "Layout"));
 	{
 		MenuBuilder.AddMenuEntry(FMainFrameCommands::Get().ResetLayout);
diff --git a/Engine/Source/Editor/UnrealEd/Classes/Settings/EditorExperimentalSettings.h b/Engine/Source/Editor/UnrealEd/Classes/Settings/EditorExperimentalSettings.h
--- a/Engine/Source/Editor/UnrealEd/Classes/Settings/EditorExperimentalSettings.h
+++ b/Engine/Source/Editor/UnrealEd/Classes/Settings/EditorExperimentalSettings.h
@@ -29,10 +29,6 @@ public:
 	UPROPERTY(EditAnywhere, config, Category=Tools, meta=(DisplayName="Editor Utility Blueprints (Blutility)"))
 	bool bEnableEditorUtilityBlueprints;
 
-	/** The Project Launcher provides advanced workflows for packaging, deploying and launching your projects. */
-	UPROPERTY(EditAnywhere, config, Category=Tools, meta=(DisplayName="Project Launcher"))
-	bool bProjectLauncher;
-
 	/** The Messaging Debugger provides a visual utility for debugging the messaging system. */
 	UPROPERTY(EditAnywhere, config, Category=Tools, meta=(DisplayName="Messaging Debugger"))
 	bool bMessagingDebugger;
diff --git a/Engine/Source/Editor/UnrealEd/Private/Kismet2/DebuggerCommands.cpp b/Engine/Source/Editor/UnrealEd/Private/Kismet2/DebuggerCommands.cpp
--- a/Engine/Source/Editor/UnrealEd/Private/Kismet2/DebuggerCommands.cpp
+++ b/Engine/Source/Editor/UnrealEd/Private/Kismet2/DebuggerCommands.cpp
@@ -81,7 +81,7 @@ public:
 	static bool RepeatLastLaunch_CanExecute();
 	static FText GetRepeatLastLaunchToolTip();
 	static FSlateIcon GetRepeatLastLaunchIcon();
-	static void OpenLauncher_Clicked();
+	static void OpenProjectLauncher_Clicked();
 	static void OpenDeviceManager_Clicked();
 
 	static FSlateIcon GetResumePlaySessionImage();
@@ -245,6 +245,7 @@ void FPlayWorldCommands::RegisterCommands()
 
 	// Launch
 	UI_COMMAND( RepeatLastLaunch, "Launch", "Launches the game on the device as the last session launched from the dropdown next to the Play on Device button on the level editor toolbar", EUserInterfaceActionType::Button, FInputGesture( EKeys::P, EModifierKey::Alt | EModifierKey::Shift ) )
+	UI_COMMAND( OpenProjectLauncher, "Project Launcher...", "Open the Project Launcher for advanced packaging, deploying and launching of your projects", EUserInterfaceActionType::Button, FInputGesture());
 	UI_COMMAND( OpenDeviceManager, "Device Manager...", "View and manage connected devices.", EUserInterfaceActionType::Button, FInputGesture() );
 }
 
@@ -328,6 +329,10 @@ void FPlayWorldCommands::BindGlobalPlayWorldCommands()
 		);
 
 	// Launch
+	ActionList.MapAction(Commands.OpenProjectLauncher,
+		FExecuteAction::CreateStatic(&FInternalPlayWorldCommandCallbacks::OpenProjectLauncher_Clicked)
+		);
+	
 	ActionList.MapAction( Commands.OpenDeviceManager,
 		FExecuteAction::CreateStatic( &FInternalPlayWorldCommandCallbacks::OpenDeviceManager_Clicked )
 		);
@@ -798,6 +803,13 @@ TSharedRef< SWidget > FPlayWorldCommands::GenerateLaunchMenuContent( TSharedRef<
 	// options section
 	MenuBuilder.BeginSection("LevelEditorLaunchOptions");
 	{
+		MenuBuilder.AddMenuEntry(FPlayWorldCommands::Get().OpenProjectLauncher,
+			NAME_None,
+			TAttribute<FText>(),
+			TAttribute<FText>(),
+			FSlateIcon(FEditorStyle::GetStyleSetName(), "Launcher.TabIcon")
+			);
+
 		MenuBuilder.AddMenuEntry( FPlayWorldCommands::Get().OpenDeviceManager,
 			NAME_None,
 			TAttribute<FText>(),
@@ -1376,19 +1388,16 @@ void FInternalPlayWorldCommandCallbacks::PlayInSettings_Clicked()
 	FModuleManager::LoadModuleChecked<ISettingsModule>("Settings").ShowViewer("Editor", "LevelEditor", "PlayIn");
 }
 
-
-void FInternalPlayWorldCommandCallbacks::OpenDeviceManager_Clicked()
+void FInternalPlayWorldCommandCallbacks::OpenProjectLauncher_Clicked()
 {
-	FGlobalTabmanager::Get()->InvokeTab(FTabId("DeviceManager"));
+	FGlobalTabmanager::Get()->InvokeTab(FTabId("ProjectLauncher"));
 }
 
-
-void FInternalPlayWorldCommandCallbacks::OpenLauncher_Clicked()
+void FInternalPlayWorldCommandCallbacks::OpenDeviceManager_Clicked()
 {
-	//FGlobalTabmanager::Get()->InvokeTab(FName("ProjectLauncher"));
+	FGlobalTabmanager::Get()->InvokeTab(FTabId("DeviceManager"));
 }
 
-
 void FInternalPlayWorldCommandCallbacks::RepeatLastLaunch_Clicked()
 {
 	const ULevelEditorPlaySettings* PlaySettings = GetDefault<ULevelEditorPlaySettings>();
diff --git a/Engine/Source/Editor/UnrealEd/Public/Kismet2/DebuggerCommands.h b/Engine/Source/Editor/UnrealEd/Public/Kismet2/DebuggerCommands.h
--- a/Engine/Source/Editor/UnrealEd/Public/Kismet2/DebuggerCommands.h
+++ b/Engine/Source/Editor/UnrealEd/Public/Kismet2/DebuggerCommands.h
@@ -66,6 +66,7 @@ public:
 
 	/** Launch on device */
 	TSharedPtr<FUICommandInfo> RepeatLastLaunch;
+	TSharedPtr<FUICommandInfo> OpenProjectLauncher;
 	TSharedPtr<FUICommandInfo> OpenDeviceManager;
 
 protected:
