make it so the game executable is built even for content-only projects if from GitHub or P4, source code is found to be present, and there is a compiler
UE-7459
#uat

diff --git a/Engine/Source/Editor/MainFrame/Private/Frame/MainFrameActions.cpp b/Engine/Source/Editor/MainFrame/Private/Frame/MainFrameActions.cpp
--- a/Engine/Source/Editor/MainFrame/Private/Frame/MainFrameActions.cpp
+++ b/Engine/Source/Editor/MainFrame/Private/Frame/MainFrameActions.cpp
@@ -355,6 +355,14 @@ void FMainFrameActionCallbacks::PackageProject( const FName InPlatformInfoName )
 	// does the project have any code?
 	FGameProjectGenerationModule& GameProjectModule = FModuleManager::LoadModuleChecked<FGameProjectGenerationModule>(TEXT("GameProjectGeneration"));
 	bool bProjectHasCode = GameProjectModule.Get().ProjectHasCodeFiles();
+	bool bContentProjectCanBeBuilt = false;
+	if (!bProjectHasCode && !FRocketSupport::IsRocket())
+	{
+		TArray<FString> OutProjectCodeFilenames;
+		IFileManager::Get().FindFilesRecursive(OutProjectCodeFilenames, *(FPaths::EngineDir() + TEXT("Source/")), TEXT("*.h"), true, false, false);
+		IFileManager::Get().FindFilesRecursive(OutProjectCodeFilenames, *(FPaths::EngineDir() + TEXT("Source/")), TEXT("*.cpp"), true, false, false);
+		bContentProjectCanBeBuilt = (OutProjectCodeFilenames.Num() > 0);
+	}
 
 	const PlatformInfo::FPlatformInfo* const PlatformInfo = PlatformInfo::FindPlatformInfo(InPlatformInfoName);
 	check(PlatformInfo);
@@ -514,7 +522,7 @@ void FMainFrameActionCallbacks::PackageProject( const FName InPlatformInfoName )
 	}
 
 	// only build if the project has code that might need to be built
-	if (bProjectHasCode && FSourceCodeNavigation::IsCompilerAvailable())
+	if ((bProjectHasCode || bContentProjectCanBeBuilt) && FSourceCodeNavigation::IsCompilerAvailable())
 	{
 		OptionalParams += TEXT(" -build");
 	}
