Removed source control icon from top-right of MainFrame.

diff --git a/Engine/Source/Developer/SourceControl/Private/SSourceControlPicker.cpp b/Engine/Source/Developer/SourceControl/Private/SSourceControlPicker.cpp
--- a/Engine/Source/Developer/SourceControl/Private/SSourceControlPicker.cpp
+++ b/Engine/Source/Developer/SourceControl/Private/SSourceControlPicker.cpp
@@ -3,7 +3,6 @@
 #include "SourceControlPrivatePCH.h"
 #include "SSourceControlPicker.h"
 #include "SourceControlModule.h"
-#include "SSourceControlStatus.h"
 #include "SSourceControlLogin.h"
 
 #if SOURCE_CONTROL_WITH_SLATE
diff --git a/Engine/Source/Developer/SourceControl/Private/SSourceControlStatus.cpp b/Engine/Source/Developer/SourceControl/Private/SSourceControlStatus.cpp
deleted file mode 100644
diff --git a/Engine/Source/Developer/SourceControl/Private/SSourceControlStatus.h b/Engine/Source/Developer/SourceControl/Private/SSourceControlStatus.h
deleted file mode 100644
diff --git a/Engine/Source/Developer/SourceControl/Private/SourceControlModule.cpp b/Engine/Source/Developer/SourceControl/Private/SourceControlModule.cpp
--- a/Engine/Source/Developer/SourceControl/Private/SourceControlModule.cpp
+++ b/Engine/Source/Developer/SourceControl/Private/SourceControlModule.cpp
@@ -6,7 +6,6 @@
 #include "Runtime/Core/Public/Features/IModularFeatures.h"
 #include "Developer/MessageLog/Public/MessageLogModule.h"
 #include "SSourceControlLogin.h"
-#include "SSourceControlStatus.h"
 
 #if WITH_EDITOR
 	#include "Runtime/Engine/Public/EngineAnalytics.h"
@@ -164,16 +163,6 @@ void FSourceControlModule::ShowLoginDialog(const FSourceControlLoginClosed& InOn
 #endif // SOURCE_CONTROL_WITH_SLATE
 }
 
-TSharedPtr<class SWidget> FSourceControlModule::CreateStatusWidget() const
-{
-#if SOURCE_CONTROL_WITH_SLATE
-	return SNew(SSourceControlStatus);
-#else
-	STUBBED("FSourceControlModule::CreateStatusWidget - no Slate");
-	return nullptr;
-#endif // SOURCE_CONTROL_WITH_SLATE
-}
-
 void FSourceControlModule::OnSourceControlDialogClosed(const TSharedRef<SWindow>& InWindow)
 {
 	SourceControlLoginWindowPtr = NULL;
diff --git a/Engine/Source/Developer/SourceControl/Private/SourceControlModule.h b/Engine/Source/Developer/SourceControl/Private/SourceControlModule.h
--- a/Engine/Source/Developer/SourceControl/Private/SourceControlModule.h
+++ b/Engine/Source/Developer/SourceControl/Private/SourceControlModule.h
@@ -31,7 +31,6 @@ public:
 	virtual ISourceControlProvider& GetProvider() const override;
 	virtual void SetProvider( const FName& InName ) override;
 	virtual void ShowLoginDialog(const FSourceControlLoginClosed& InOnSourceControlLoginClosed, ELoginWindowMode::Type InLoginWindowMode, EOnLoginWindowStartup::Type InOnLoginWindowStartup = EOnLoginWindowStartup::ResetProviderToNone) override;
-	virtual TSharedPtr<class SWidget> CreateStatusWidget() const override;
 	virtual bool GetUseGlobalSettings() const override;
 	virtual void SetUseGlobalSettings(bool bIsUseGlobalSettings) override;
 
diff --git a/Engine/Source/Developer/SourceControl/Public/ISourceControlModule.h b/Engine/Source/Developer/SourceControl/Public/ISourceControlModule.h
--- a/Engine/Source/Developer/SourceControl/Public/ISourceControlModule.h
+++ b/Engine/Source/Developer/SourceControl/Public/ISourceControlModule.h
@@ -98,12 +98,6 @@ public:
 	 */
 	virtual void ShowLoginDialog(const FSourceControlLoginClosed& InOnSourceControlLoginClosed, ELoginWindowMode::Type InLoginWindowMode, EOnLoginWindowStartup::Type InOnLoginWindowStartup = EOnLoginWindowStartup::ResetProviderToNone) = 0;
 
-	/**
-	 * Create a status widget, used to display passive source control status
-	 * @returns a new status widget
-	 */
-	virtual TSharedPtr<class SWidget> CreateStatusWidget() const = 0;
-
 	/**
 	 * Get whether we should use global or per-project settings
 	 * @return true if we should use global settings
diff --git a/Engine/Source/Editor/MainFrame/Private/MainFrameModule.cpp b/Engine/Source/Editor/MainFrame/Private/MainFrameModule.cpp
--- a/Engine/Source/Editor/MainFrame/Private/MainFrameModule.cpp
+++ b/Engine/Source/Editor/MainFrame/Private/MainFrameModule.cpp
@@ -482,13 +482,6 @@ TSharedRef<SWidget> FMainFrameModule::MakeDeveloperTools() const
 
 	bool bUseSuperSearch = true;
 
-	TSharedPtr<SWidget> StatusWidget = ISourceControlModule::Get().CreateStatusWidget();
-	// source control module can be compiled without UI support, provide a fallback
-	if (!StatusWidget.IsValid())
-	{
-		StatusWidget = SNew( SSpacer );
-	}
-	
 	// Invisible border, so that we can animate our box panel size
 	return SNew( SBorder )
 		.Visibility( EVisibility::SelfHitTestInvisible )
@@ -542,15 +535,6 @@ TSharedRef<SWidget> FMainFrameModule::MakeDeveloperTools() const
 						.ToolTipText(LOCTEXT("SearchEditor", "Search the Editor"))
 				]
 */
-			+SHorizontalBox::Slot()
-				.AutoWidth()
-				.VAlign(VAlign_Center)
-				.Padding( 6.0f, 0.0f, 2.0f, 0.0f )
-				[
-					// STATUS BUTTON
-					StatusWidget.ToSharedRef()
-				]
-
 			// Editor live streaming toggle button
 			+SHorizontalBox::Slot()
 				.AutoWidth()
