bShowProjectMenus option for editor

Merging CL 2413753

//depot/UE4-UT/...

to //depot/UE4/...

diff --git a/Engine/Config/BaseEditorUserSettings.ini b/Engine/Config/BaseEditorUserSettings.ini
--- a/Engine/Config/BaseEditorUserSettings.ini
+++ b/Engine/Config/BaseEditorUserSettings.ini
@@ -77,7 +77,7 @@ bEnableWindowAnimations=False
 bShowFriendlyNames=True
 ;If true, the Editor Preferences and Project Settings menu entries in the main menu will be expanded with sub-menus for each settings section
 bExpandConfigurationMenus=False
-
+bShowProjectMenus=True
 
 [/Script/Levels.LevelBrowserSettings]
 ; True if the actor count is displayed in the slate level browser
diff --git a/Engine/Source/Editor/EditorStyle/Public/Classes/EditorStyleSettings.h b/Engine/Source/Editor/EditorStyle/Public/Classes/EditorStyleSettings.h
--- a/Engine/Source/Editor/EditorStyle/Public/Classes/EditorStyleSettings.h
+++ b/Engine/Source/Editor/EditorStyle/Public/Classes/EditorStyleSettings.h
@@ -78,6 +78,10 @@ public:
 	UPROPERTY(EditAnywhere, config, Category=UserInterface, AdvancedDisplay)
 	uint32 bExpandConfigurationMenus:1;
 
+	/** When enabled, the Editor Preferences and Project Settings menu items in the main menu will be expanded with sub-menus for each settings section. */
+	UPROPERTY(config)
+	uint32 bShowProjectMenus : 1;
+
 	/** The display mode for timestamps in the output log */
 	UPROPERTY(EditAnywhere, config, Category=UserInterface)
 	TEnumAsByte<ELogTimes::Type> LogTimestampMode;
diff --git a/Engine/Source/Editor/MainFrame/Private/Menus/MainMenu.cpp b/Engine/Source/Editor/MainFrame/Private/Menus/MainMenu.cpp
--- a/Engine/Source/Editor/MainFrame/Private/Menus/MainMenu.cpp
+++ b/Engine/Source/Editor/MainFrame/Private/Menus/MainMenu.cpp
@@ -417,18 +417,22 @@ TSharedRef< SWidget > FMainMenu::MakeMainTabMenu( const TSharedPtr<FTabManager>&
 
 		static void FillRecentFileAndExitMenuItems( FMenuBuilder& MenuBuilder )
 		{
-			MenuBuilder.BeginSection( "FileRecentFiles" );
+			if (GetDefault<UEditorStyleSettings>()->bShowProjectMenus)
 			{
-				if ( FMainFrameActionCallbacks::ProjectNames.Num() > 0 )
+				MenuBuilder.BeginSection("FileRecentFiles");
 				{
-					MenuBuilder.AddSubMenu(
-						LOCTEXT("SwitchProjectSubMenu", "Recent Projects"),
-						LOCTEXT("SwitchProjectSubMenu_ToolTip", "Select a project to switch to"),
-						FNewMenuDelegate::CreateStatic( &FRecentProjectsMenu::MakeMenu ), false, FSlateIcon(FEditorStyle::GetStyleSetName(), "MainFrame.RecentProjects")
-					);
+					if (FMainFrameActionCallbacks::ProjectNames.Num() > 0)
+					{
+						MenuBuilder.AddSubMenu(
+							LOCTEXT("SwitchProjectSubMenu", "Recent Projects"),
+							LOCTEXT("SwitchProjectSubMenu_ToolTip", "Select a project to switch to"),
+							FNewMenuDelegate::CreateStatic(&FRecentProjectsMenu::MakeMenu), false, FSlateIcon(FEditorStyle::GetStyleSetName(), "MainFrame.RecentProjects")
+							);
+					}
 				}
+				MenuBuilder.EndSection();
 			}
-			MenuBuilder.EndSection();
+
 #if !PLATFORM_MAC // Handled by app's menu in menu bar
 			MenuBuilder.AddMenuSeparator();
 			MenuBuilder.AddMenuEntry( FMainFrameCommands::Get().Exit, "Exit" );
@@ -444,11 +448,14 @@ TSharedRef< SWidget > FMainMenu::MakeMainTabMenu( const TSharedPtr<FTabManager>&
 	
 		IMainFrameModule& MainFrameModule = FModuleManager::GetModuleChecked<IMainFrameModule>("MainFrame");
 
-		Extender->AddMenuExtension(
-			"FileLoadAndSave",
-			EExtensionHook::After,
-			MainFrameModule.GetMainFrameCommandBindings(),
-			FMenuExtensionDelegate::CreateStatic( &Local::FillProjectMenuItems ) );
+		if (GetDefault<UEditorStyleSettings>()->bShowProjectMenus)
+		{
+			Extender->AddMenuExtension(
+				"FileLoadAndSave",
+				EExtensionHook::After,
+				MainFrameModule.GetMainFrameCommandBindings(),
+				FMenuExtensionDelegate::CreateStatic(&Local::FillProjectMenuItems));
+		}
 
 		Extender->AddMenuExtension(
 			"FileLoadAndSave",
