Adjust bShowProjectMenus usage

Merging

//depot/UE4-UT/...

to //depot/UE4/...

diff --git a/Engine/Source/Editor/MainFrame/Private/Menus/MainMenu.cpp b/Engine/Source/Editor/MainFrame/Private/Menus/MainMenu.cpp
--- a/Engine/Source/Editor/MainFrame/Private/Menus/MainMenu.cpp
+++ b/Engine/Source/Editor/MainFrame/Private/Menus/MainMenu.cpp
@@ -417,21 +417,18 @@ TSharedRef< SWidget > FMainMenu::MakeMainTabMenu( const TSharedPtr<FTabManager>&
 
 		static void FillRecentFileAndExitMenuItems( FMenuBuilder& MenuBuilder )
 		{
-			if (GetDefault<UEditorStyleSettings>()->bShowProjectMenus)
+			MenuBuilder.BeginSection("FileRecentFiles");
 			{
-				MenuBuilder.BeginSection("FileRecentFiles");
+				if (GetDefault<UEditorStyleSettings>()->bShowProjectMenus && FMainFrameActionCallbacks::ProjectNames.Num() > 0)
 				{
-					if (FMainFrameActionCallbacks::ProjectNames.Num() > 0)
-					{
-						MenuBuilder.AddSubMenu(
-							LOCTEXT("SwitchProjectSubMenu", "Recent Projects"),
-							LOCTEXT("SwitchProjectSubMenu_ToolTip", "Select a project to switch to"),
-							FNewMenuDelegate::CreateStatic(&FRecentProjectsMenu::MakeMenu), false, FSlateIcon(FEditorStyle::GetStyleSetName(), "MainFrame.RecentProjects")
-							);
-					}
+					MenuBuilder.AddSubMenu(
+						LOCTEXT("SwitchProjectSubMenu", "Recent Projects"),
+						LOCTEXT("SwitchProjectSubMenu_ToolTip", "Select a project to switch to"),
+						FNewMenuDelegate::CreateStatic(&FRecentProjectsMenu::MakeMenu), false, FSlateIcon(FEditorStyle::GetStyleSetName(), "MainFrame.RecentProjects")
+						);
 				}
-				MenuBuilder.EndSection();
 			}
+			MenuBuilder.EndSection();
 
 #if !PLATFORM_MAC // Handled by app's menu in menu bar
 			MenuBuilder.AddMenuSeparator();
