[INTEGRATE] Change 2417361 by [REDACTED_PERSON][REDACTED_HANDLE] on 2015/01/23 17:29:38

	"Add code to project" should no longer be greyed out when the user doesn't have an available compiler

diff --git a/Engine/Source/Editor/MainFrame/Private/Frame/MainFrameActions.cpp b/Engine/Source/Editor/MainFrame/Private/Frame/MainFrameActions.cpp
--- a/Engine/Source/Editor/MainFrame/Private/Frame/MainFrameActions.cpp
+++ b/Engine/Source/Editor/MainFrame/Private/Frame/MainFrameActions.cpp
@@ -62,7 +62,7 @@ void FMainFrameCommands::RegisterCommands()
 	ActionList->MapAction( OpenProject, FExecuteAction::CreateStatic( &FMainFrameActionCallbacks::NewProject, true, false), DefaultExecuteAction );
 
 	UI_COMMAND( AddCodeToProject, "Add Code to Project...", "Adds C++ code to the project. The code can only be compiled if you have an appropriate C++ compiler installed.", EUserInterfaceActionType::Button, FInputGesture() );
-	ActionList->MapAction( AddCodeToProject, FExecuteAction::CreateStatic( &FMainFrameActionCallbacks::AddCodeToProject ), FCanExecuteAction::CreateStatic( &FSourceCodeNavigation::IsCompilerAvailable ) );
+	ActionList->MapAction( AddCodeToProject, FExecuteAction::CreateStatic( &FMainFrameActionCallbacks::AddCodeToProject ));
 
 	UI_COMMAND( RefreshCodeProject, "Refresh code project", "Refreshes your C++ code project.", EUserInterfaceActionType::Button, FInputGesture() );
 	ActionList->MapAction( RefreshCodeProject, FExecuteAction::CreateStatic( &FMainFrameActionCallbacks::RefreshCodeProject ), DefaultExecuteAction );
