Workaround for crash clicking File menu

diff --git a/Engine/Source/Editor/MainFrame/Private/Menus/MainMenu.cpp b/Engine/Source/Editor/MainFrame/Private/Menus/MainMenu.cpp
--- a/Engine/Source/Editor/MainFrame/Private/Menus/MainMenu.cpp
+++ b/Engine/Source/Editor/MainFrame/Private/Menus/MainMenu.cpp
@@ -378,13 +378,13 @@ TSharedRef< SWidget > FMainMenu::MakeMainTabMenu( const TSharedPtr<FTabManager>&
 					LOCTEXT("PackageProjectSubMenuToolTip", "Compile, cook and package your project and its content for distribution."),
 					FNewMenuDelegate::CreateStatic( &FPackageProjectMenu::MakeMenu ), false, FSlateIcon(FEditorStyle::GetStyleSetName(), "MainFrame.PackageProject")
 				);
-
+				/*
 				MenuBuilder.AddMenuEntry( FMainFrameCommands::Get().LocalizeProject,
 					NAME_None,
 					TAttribute<FText>(),
 					LOCTEXT("LocalizeProjectToolTip", "Gather text from your project and import/export translations.")
 					);
-
+					*/
 				/*
 				MenuBuilder.AddSubMenu(
 					LOCTEXT("CookProjectSubMenuLabel", "Cook Project"),
