Added support for cook all and cook maps only from packaging settings.
Added list of maps to cook for the packaging settings.

diff --git a/Engine/Source/Editor/MainFrame/Private/Frame/MainFrameActions.cpp b/Engine/Source/Editor/MainFrame/Private/Frame/MainFrameActions.cpp
--- a/Engine/Source/Editor/MainFrame/Private/Frame/MainFrameActions.cpp
+++ b/Engine/Source/Editor/MainFrame/Private/Frame/MainFrameActions.cpp
@@ -315,7 +315,7 @@ void FMainFrameActionCallbacks::CookContent( const FName InPlatformInfoName )
 	}
 
 	FString ProjectPath = FPaths::IsProjectFilePathSet() ? FPaths::ConvertRelativePathToFull(FPaths::GetProjectFilePath()) : FPaths::RootDir() / FApp::GetGameName() / FApp::GetGameName() + TEXT(".uproject");
-	FString CommandLine = FString::Printf(TEXT("BuildCookRun %s%s -nop4 -project=\"%s\" -cook -allmaps -%s -ue4exe=%s %s"),
+	FString CommandLine = FString::Printf(TEXT("BuildCookRun %s%s -nop4 -project=\"%s\" -cook -%s -ue4exe=%s %s"),
 		FRocketSupport::IsRocket() ? TEXT("-rocket -nocompile") : TEXT("-nocompileeditor"),
 		FApp::IsEngineInstalled() ? TEXT(" -installed") : TEXT(""),
 		*ProjectPath,
@@ -481,6 +481,27 @@ void FMainFrameActionCallbacks::PackageProject( const FName InPlatformInfoName )
 		OptionalParams += TEXT(" -clean");
 	}
 
+	if ( PackagingSettings->bCookAll )
+	{
+		OptionalParams += TEXT(" -CookAll");
+		// maps only flag only affects cook all
+		if ( PackagingSettings->bCookMapsOnly )
+		{
+			OptionalParams += TEXT(" -CookMapsOnly");
+		}
+	}
+	else ( PackagingSettings->MapsToCook.Num() )
+	{
+		OptionalParams += TEXT(" -mapstocook=");
+		for ( const auto& Map : PackagingSettings->MapsToCook )
+		{
+			OptionalParams += FString::Printf(TEXT("%s+"), *Map.FilePath);
+		}
+		OptionalParams.RemoveFromEnd("+");
+	}
+	
+	
+
 	if (PackagingSettings->UsePakFile)
 	{
 	  if (PlatformInfo->TargetPlatformName != FName("HTML5")) 
@@ -560,7 +581,7 @@ void FMainFrameActionCallbacks::PackageProject( const FName InPlatformInfoName )
 	Configuration = Configuration.Replace(TEXT("PPBC_"), TEXT(""));
 
 	FString ProjectPath = FPaths::IsProjectFilePathSet() ? FPaths::ConvertRelativePathToFull(FPaths::GetProjectFilePath()) : FPaths::RootDir() / FApp::GetGameName() / FApp::GetGameName() + TEXT(".uproject");
-	FString CommandLine = FString::Printf(TEXT("BuildCookRun %s%s -nop4 -project=\"%s\" -cook -allmaps -stage -archive -archivedirectory=\"%s\" -package -%s -clientconfig=%s -ue4exe=%s %s -utf8output"),
+	FString CommandLine = FString::Printf(TEXT("BuildCookRun %s%s -nop4 -project=\"%s\" -cook -stage -archive -archivedirectory=\"%s\" -package -%s -clientconfig=%s -ue4exe=%s %s -utf8output"),
 		FRocketSupport::IsRocket() ? TEXT( "-rocket -nocompile" ) : TEXT( "-nocompileeditor" ),
 		FApp::IsEngineInstalled() ? TEXT(" -installed") : TEXT(""),
 		*ProjectPath,
diff --git a/Engine/Source/Editor/UnrealEd/Classes/Settings/ProjectPackagingSettings.h b/Engine/Source/Editor/UnrealEd/Classes/Settings/ProjectPackagingSettings.h
--- a/Engine/Source/Editor/UnrealEd/Classes/Settings/ProjectPackagingSettings.h
+++ b/Engine/Source/Editor/UnrealEd/Classes/Settings/ProjectPackagingSettings.h
@@ -108,13 +108,31 @@ public:
 	UPROPERTY(config, EditAnywhere, Category=Packaging, AdvancedDisplay, meta=(DisplayName="Localizations to Package"))
 	TArray<FString> CulturesToStage;
 
+	/**
+	 * Cook all things in the project content directory
+	 */
+	UPROPERTY(config, EditAnywhere, Category=Packaging, AdvancedDisplay, meta=(DisplayName="Cook everything in the project content directory (ignore list of maps below)"))
+	bool bCookAll;
+
+	/**
+	 * Cook only maps (this only affects the cookall flag)
+	 */
+	UPROPERTY(config, EditAnywhere, Category=Packaging, AdvancedDisplay, meta=(DisplayName="Cook only maps (this only affects cookall)"))
+	bool bCookMapsOnly;
+	
+	/**
+	 * List of maps to include when no other map list is specified on commandline
+	 */
+	UPROPERTY(config, EditAnywhere, Category=Packaging, AdvancedDisplay, meta=(DisplayName="List of maps to include in a packaged build", RelativeToGameContentDir))
+	TArray<FFilePath> MapsToCook;	
+
 	/**
 	 * Directories containing .uasset files that should always be cooked regardless of whether they're referenced by anything in your project
 	 * Note: These paths are relative to your project Content directory
 	 */
 	UPROPERTY(config, EditAnywhere, Category=Packaging, AdvancedDisplay, meta=(DisplayName="Additional Asset Directories to Cook", RelativeToGameContentDir))
 	TArray<FDirectoryPath> DirectoriesToAlwaysCook;
-
+	
 	/**
 	 * Directories containing files that should always be added to the .pak file (if using a .pak file; otherwise they're copied as individual files)
 	 * This is used to stage additional files that you manually load via the UFS (Unreal File System) file IO API
@@ -131,6 +149,8 @@ public:
 	UPROPERTY(config, EditAnywhere, Category=Packaging, AdvancedDisplay, meta=(DisplayName="Additional Non-Asset Directories To Copy", RelativeToGameContentDir))
 	TArray<FDirectoryPath> DirectoriesToAlwaysStageAsNonUFS;	
 
+
+	
 public:
 
 	// UObject Interface
diff --git a/Engine/Source/Programs/AutomationTool/Scripts/BuildCookRun.Automation.cs b/Engine/Source/Programs/AutomationTool/Scripts/BuildCookRun.Automation.cs
--- a/Engine/Source/Programs/AutomationTool/Scripts/BuildCookRun.Automation.cs
+++ b/Engine/Source/Programs/AutomationTool/Scripts/BuildCookRun.Automation.cs
@@ -58,28 +58,6 @@ public class BuildCookRun : BuildCommand
 			RawProjectPath: ProjectPath
 		);
 
-		// Initialize map
-		var Map = ParseParamValue("map");
-		if (Map == null)
-		{
-			LogVerbose("-map command line param not found, trying to find DefaultMap in INI.");
-			Map = GetDefaultMap(Params);
-		}
-
-		if (!String.IsNullOrEmpty(Map))
-		{
-			if (ParseParam("allmaps"))
-			{
-				Log("Cooking all maps");
-			}
-			else
-			{
-				Params.MapsToCook = new ParamList<string>(Map);
-			}
-
-			Params.MapToRun = GetFirstMap(Map);
-		}
-
 		// @rocket hack: non-code projects cannot run in Debug
 		if (Params.Rocket && !ProjectUtils.IsCodeBasedUProjectFile(ProjectPath))
 		{
