Fix compilation error from cl 2430466

diff --git a/Engine/Source/Editor/MainFrame/Private/Frame/MainFrameActions.cpp b/Engine/Source/Editor/MainFrame/Private/Frame/MainFrameActions.cpp
--- a/Engine/Source/Editor/MainFrame/Private/Frame/MainFrameActions.cpp
+++ b/Engine/Source/Editor/MainFrame/Private/Frame/MainFrameActions.cpp
@@ -490,7 +490,7 @@ void FMainFrameActionCallbacks::PackageProject( const FName InPlatformInfoName )
 			OptionalParams += TEXT(" -CookMapsOnly");
 		}
 	}
-	else ( PackagingSettings->MapsToCook.Num() )
+	else if ( PackagingSettings->MapsToCook.Num() )
 	{
 		OptionalParams += TEXT(" -mapstocook=");
 		for ( const auto& Map : PackagingSettings->MapsToCook )
