UE-7755: UBT reporting VS2012 as an unsupported compiler even when UAT was passed -NoCompile

There were two problems with this. One, I had to suppress registering build platforms in UBT if the project has no code. To achieve this I added new flag -nocodeproject and I'm passing it down from the editor to UAT.

The second one emerged after fixing previous one. The Editor was adding -build flag to the UAT if it found any of the VS 2012 or 2013. -build flag is making UAT to try to build, which fails as it is not supported on VS 2012-only system. During the discussion with platform team we came up with the solution to not check for VS 2012 in editor, as we only support it for XDK development, which requires both VSes on the system (2012 and 2013). Hence flag -build should be appended to the UAT command line if and only if there is VS2013 on the system.

diff --git a/Engine/Plugins/Developer/VisualStudioSourceCodeAccess/Source/VisualStudioSourceCodeAccess/Private/VisualStudioSourceCodeAccessor.cpp b/Engine/Plugins/Developer/VisualStudioSourceCodeAccess/Source/VisualStudioSourceCodeAccess/Private/VisualStudioSourceCodeAccessor.cpp
--- a/Engine/Plugins/Developer/VisualStudioSourceCodeAccess/Source/VisualStudioSourceCodeAccess/Private/VisualStudioSourceCodeAccessor.cpp
+++ b/Engine/Plugins/Developer/VisualStudioSourceCodeAccess/Source/VisualStudioSourceCodeAccess/Private/VisualStudioSourceCodeAccessor.cpp
@@ -79,10 +79,6 @@ void FVisualStudioSourceCodeAccessor::Startup()
 
 	// Preferential order of VS versions
 	AddVisualStudioVersion(12); // Visual Studio 2013
-	if (!FRocketSupport::IsRocket())
-	{
-		AddVisualStudioVersion(11); // Visual Studio 2012
-	}
 }
 
 void FVisualStudioSourceCodeAccessor::Shutdown()
diff --git a/Engine/Source/Editor/MainFrame/Private/Frame/MainFrameActions.cpp b/Engine/Source/Editor/MainFrame/Private/Frame/MainFrameActions.cpp
--- a/Engine/Source/Editor/MainFrame/Private/Frame/MainFrameActions.cpp
+++ b/Engine/Source/Editor/MainFrame/Private/Frame/MainFrameActions.cpp
@@ -272,7 +272,27 @@ void FMainFrameActionCallbacks::AddCodeToProject()
 	FGameProjectGenerationModule::Get().OpenAddCodeToProjectDialog();
 }
 
-void FMainFrameActionCallbacks::CookContent( const FName InPlatformInfoName )
+/**
+ * Gets compilation flags for UAT for this system.
+ */
+const TCHAR* GetUATCompilationFlags()
+{
+	return FRocketSupport::IsRocket() || !FSourceCodeNavigation::IsCompilerAvailable()
+		? TEXT("-nocompile")
+		: TEXT("-nocompileeditor");
+}
+
+/**
+ * Gets -nocodeproject flag for UAT if it's needed.
+ */
+const TCHAR* GetUATNoCodeProjectFlag()
+{
+	IHotReloadInterface& HotReloadSupport = FModuleManager::LoadModuleChecked<IHotReloadInterface>("HotReload");
+	// If there is no game module loaded, source code is not available.
+	return !HotReloadSupport.IsAnyGameModuleLoaded() ? TEXT(" -nocodeproject") : TEXT("");
+}
+
+void FMainFrameActionCallbacks::CookContent(const FName InPlatformInfoName)
 {
 	const PlatformInfo::FPlatformInfo* const PlatformInfo = PlatformInfo::FindPlatformInfo(InPlatformInfoName);
 	check(PlatformInfo);
@@ -315,8 +335,10 @@ void FMainFrameActionCallbacks::CookContent( const FName InPlatformInfoName )
 	}
 
 	FString ProjectPath = FPaths::IsProjectFilePathSet() ? FPaths::ConvertRelativePathToFull(FPaths::GetProjectFilePath()) : FPaths::RootDir() / FApp::GetGameName() / FApp::GetGameName() + TEXT(".uproject");
-	FString CommandLine = FString::Printf(TEXT("BuildCookRun %s%s -nop4 -project=\"%s\" -cook -%s -ue4exe=%s %s"),
-		FRocketSupport::IsRocket() ? TEXT("-rocket -nocompile") : TEXT("-nocompileeditor"),
+	FString CommandLine = FString::Printf(TEXT("BuildCookRun %s%s%s%s -nop4 -project=\"%s\" -cook -%s -ue4exe=%s %s"),
+		FRocketSupport::IsRocket() ? TEXT("-rocket ") : TEXT(""),
+		GetUATCompilationFlags(),
+		GetUATNoCodeProjectFlag(),
 		FApp::IsEngineInstalled() ? TEXT(" -installed") : TEXT(""),
 		*ProjectPath,
 		*PlatformInfo->TargetPlatformName.ToString(),
@@ -581,8 +603,10 @@ void FMainFrameActionCallbacks::PackageProject( const FName InPlatformInfoName )
 	Configuration = Configuration.Replace(TEXT("PPBC_"), TEXT(""));
 
 	FString ProjectPath = FPaths::IsProjectFilePathSet() ? FPaths::ConvertRelativePathToFull(FPaths::GetProjectFilePath()) : FPaths::RootDir() / FApp::GetGameName() / FApp::GetGameName() + TEXT(".uproject");
-	FString CommandLine = FString::Printf(TEXT("BuildCookRun %s%s -nop4 -project=\"%s\" -cook -stage -archive -archivedirectory=\"%s\" -package -%s -clientconfig=%s -ue4exe=%s %s -utf8output"),
-		FRocketSupport::IsRocket() ? TEXT( "-rocket -nocompile" ) : TEXT( "-nocompileeditor" ),
+	FString CommandLine = FString::Printf(TEXT("BuildCookRun %s%s%s%s -nop4 -project=\"%s\" -cook -stage -archive -archivedirectory=\"%s\" -package -%s -clientconfig=%s -ue4exe=%s %s -utf8output"),
+		FRocketSupport::IsRocket() ? TEXT("-rocket ") : TEXT(""),
+		GetUATCompilationFlags(),
+		GetUATNoCodeProjectFlag(),
 		FApp::IsEngineInstalled() ? TEXT(" -installed") : TEXT(""),
 		*ProjectPath,
 		*PackagingSettings->StagingDirectory.Path,
diff --git a/Engine/Source/Programs/AutomationTool/Automation.cs b/Engine/Source/Programs/AutomationTool/Automation.cs
--- a/Engine/Source/Programs/AutomationTool/Automation.cs
+++ b/Engine/Source/Programs/AutomationTool/Automation.cs
@@ -105,6 +105,7 @@ namespace AutomationTool
 		public static CommandLineArg Preprocess = new CommandLineArg("-Preprocess");
 		public static CommandLineArg NoCompile = new CommandLineArg("-NoCompile");
 		public static CommandLineArg NoCompileEditor = new CommandLineArg("-NoCompileEditor");
+		public static CommandLineArg NoCodeProject = new CommandLineArg("-NoCodeProject");
 		public static CommandLineArg Help = new CommandLineArg("-Help");
 		public static CommandLineArg List = new CommandLineArg("-List");
 		public static CommandLineArg Rocket = new CommandLineArg("-Rocket");
diff --git a/Engine/Source/Programs/AutomationTool/ProjectUtils.cs b/Engine/Source/Programs/AutomationTool/ProjectUtils.cs
--- a/Engine/Source/Programs/AutomationTool/ProjectUtils.cs
+++ b/Engine/Source/Programs/AutomationTool/ProjectUtils.cs
@@ -357,6 +357,12 @@ namespace AutomationTool
 		private static void CompileAndLoadTargetsAssembly(ProjectProperties Properties, string TargetsDllFilename, bool DoNotCompile, List<string> TargetScripts)
 		{
 			CommandUtils.Log("Compiling targets DLL: {0}", TargetsDllFilename);
+
+			if (!DoNotCompile && GlobalCommandLine.NoCodeProject)
+			{
+				//throw new AutomationException("Building is not supported when -nocodeproject flag is provided.");
+			}
+
 			var ReferencedAssemblies = new List<string>() 
 					{ 
 						"System.dll", 
diff --git a/Engine/Source/Programs/AutomationTool/ScriptCompiler.cs b/Engine/Source/Programs/AutomationTool/ScriptCompiler.cs
--- a/Engine/Source/Programs/AutomationTool/ScriptCompiler.cs
+++ b/Engine/Source/Programs/AutomationTool/ScriptCompiler.cs
@@ -70,7 +70,7 @@ namespace AutomationTool
 			}
 			// Register all the classes inside UBT
 			Log.TraceVerbose("Registering UBT Classes.");
-			UnrealBuildTool.UnrealBuildTool.RegisterAllUBTClasses();
+			UnrealBuildTool.UnrealBuildTool.RegisterAllUBTClasses(GlobalCommandLine.NoCodeProject);
 			Environment.CurrentDirectory = OldCWD;
 
 			// Compile only if not disallowed.
diff --git a/Engine/Source/Programs/AutomationTool/UE4Build.cs b/Engine/Source/Programs/AutomationTool/UE4Build.cs
--- a/Engine/Source/Programs/AutomationTool/UE4Build.cs
+++ b/Engine/Source/Programs/AutomationTool/UE4Build.cs
@@ -1123,6 +1123,11 @@ namespace AutomationTool
 
 		public bool CanUseXGE(UnrealBuildTool.UnrealTargetPlatform Platform)
 		{
+			if (!UnrealBuildTool.UEBuildPlatform.BuildPlatformDictionary.ContainsKey(Platform))
+			{
+				return false;
+			}
+
 			return UnrealBuildTool.UEBuildPlatform.GetBuildPlatform(Platform).CanUseXGE();
 		}
 
@@ -1152,6 +1157,11 @@ namespace AutomationTool
 		/// <param name="InForceNoXGE">If true will force XGE off</param>
 		public void Build(BuildAgenda Agenda, bool? InDeleteBuildProducts = null, bool InUpdateVersionFiles = true, bool InForceNoXGE = false, bool InForceNonUnity = false, bool InForceUnity = false, Dictionary<UnrealBuildTool.UnrealTargetPlatform, Dictionary<string, string>> PlatformEnvVars = null)
 		{
+			if (GlobalCommandLine.NoCodeProject)
+			{
+				throw new AutomationException("Building is not supported when -nocodeproject flag is provided.");
+			}
+
 			if (!CmdEnv.HasCapabilityToCompile)
 			{
 				throw new AutomationException("You are attempting to compile on a machine that does not have a supported compiler!");
diff --git a/Engine/Source/Programs/UnrealBuildTool/System/UnrealBuildTool.cs b/Engine/Source/Programs/UnrealBuildTool/System/UnrealBuildTool.cs
--- a/Engine/Source/Programs/UnrealBuildTool/System/UnrealBuildTool.cs
+++ b/Engine/Source/Programs/UnrealBuildTool/System/UnrealBuildTool.cs
@@ -443,7 +443,7 @@ namespace UnrealBuildTool
             return false;
         }
 
-        public static void RegisterAllUBTClasses()
+        public static void RegisterAllUBTClasses(bool bSkipBuildPlatforms = false)
         {
             // Find and register all tool chains and build platforms that are present
             Assembly UBTAssembly = Assembly.GetExecutingAssembly();
@@ -453,16 +453,19 @@ namespace UnrealBuildTool
 
                 List<System.Type> ProjectGeneratorList = new List<System.Type>();
                 var AllTypes = UBTAssembly.GetTypes();
-                // register all build platforms first, since they implement SDK-switching logic that can set environment variables
-                foreach (var CheckType in AllTypes)
+                if (!bSkipBuildPlatforms)
                 {
-                    if (CheckType.IsClass && !CheckType.IsAbstract)
+                    // register all build platforms first, since they implement SDK-switching logic that can set environment variables
+                    foreach (var CheckType in AllTypes)
                     {
-                        if (Utils.ImplementsInterface<IUEBuildPlatform>(CheckType))
+                        if (CheckType.IsClass && !CheckType.IsAbstract)
                         {
-                            Log.TraceVerbose("    Registering build platform: {0}", CheckType.ToString());
-                            var TempInst = (UEBuildPlatform)(UBTAssembly.CreateInstance(CheckType.FullName, true));
-                            TempInst.RegisterBuildPlatform();
+                            if (Utils.ImplementsInterface<IUEBuildPlatform>(CheckType))
+                            {
+                                Log.TraceVerbose("    Registering build platform: {0}", CheckType.ToString());
+                                var TempInst = (UEBuildPlatform)(UBTAssembly.CreateInstance(CheckType.FullName, true));
+                                TempInst.RegisterBuildPlatform();
+                            }
                         }
                     }
                 }
