Speed up packaging steps by removing some usused code that was recursively searching for source.

diff --git a/Engine/Source/Editor/MainFrame/Private/Frame/MainFrameActions.cpp b/Engine/Source/Editor/MainFrame/Private/Frame/MainFrameActions.cpp
--- a/Engine/Source/Editor/MainFrame/Private/Frame/MainFrameActions.cpp
+++ b/Engine/Source/Editor/MainFrame/Private/Frame/MainFrameActions.cpp
@@ -380,14 +380,6 @@ void FMainFrameActionCallbacks::PackageProject( const FName InPlatformInfoName )
 	// does the project have any code?
 	FGameProjectGenerationModule& GameProjectModule = FModuleManager::LoadModuleChecked<FGameProjectGenerationModule>(TEXT("GameProjectGeneration"));
 	bool bProjectHasCode = GameProjectModule.Get().ProjectHasCodeFiles();
-	bool bContentProjectCanBeBuilt = false;
-	if (!bProjectHasCode && !FRocketSupport::IsRocket())
-	{
-		TArray<FString> OutProjectCodeFilenames;
-		IFileManager::Get().FindFilesRecursive(OutProjectCodeFilenames, *(FPaths::EngineDir() + TEXT("Source/")), TEXT("*.h"), true, false, false);
-		IFileManager::Get().FindFilesRecursive(OutProjectCodeFilenames, *(FPaths::EngineDir() + TEXT("Source/")), TEXT("*.cpp"), true, false, false);
-		bContentProjectCanBeBuilt = (OutProjectCodeFilenames.Num() > 0);
-	}
 
 	const PlatformInfo::FPlatformInfo* const PlatformInfo = PlatformInfo::FindPlatformInfo(InPlatformInfoName);
 	check(PlatformInfo);
