fix for UAT change which requires -platform= instead of -<platform>
UE-9448

diff --git a/Engine/Source/Editor/MainFrame/Private/Frame/MainFrameActions.cpp b/Engine/Source/Editor/MainFrame/Private/Frame/MainFrameActions.cpp
--- a/Engine/Source/Editor/MainFrame/Private/Frame/MainFrameActions.cpp
+++ b/Engine/Source/Editor/MainFrame/Private/Frame/MainFrameActions.cpp
@@ -336,7 +336,7 @@ void FMainFrameActionCallbacks::CookContent(const FName InPlatformInfoName)
 	}
 
 	FString ProjectPath = FPaths::IsProjectFilePathSet() ? FPaths::ConvertRelativePathToFull(FPaths::GetProjectFilePath()) : FPaths::RootDir() / FApp::GetGameName() / FApp::GetGameName() + TEXT(".uproject");
-	FString CommandLine = FString::Printf(TEXT("BuildCookRun %s%s%s%s -nop4 -project=\"%s\" -cook -%s -ue4exe=%s %s"),
+	FString CommandLine = FString::Printf(TEXT("BuildCookRun %s%s%s%s -nop4 -project=\"%s\" -cook -platform=%s -ue4exe=%s %s"),
 		FRocketSupport::IsRocket() ? TEXT("-rocket ") : TEXT(""),
 		GetUATCompilationFlags(),
 		GetUATNoCodeProjectFlag(),
@@ -596,7 +596,7 @@ void FMainFrameActionCallbacks::PackageProject( const FName InPlatformInfoName )
 	Configuration = Configuration.Replace(TEXT("PPBC_"), TEXT(""));
 
 	FString ProjectPath = FPaths::IsProjectFilePathSet() ? FPaths::ConvertRelativePathToFull(FPaths::GetProjectFilePath()) : FPaths::RootDir() / FApp::GetGameName() / FApp::GetGameName() + TEXT(".uproject");
-	FString CommandLine = FString::Printf(TEXT("BuildCookRun %s%s%s%s -nop4 -project=\"%s\" -cook -stage -archive -archivedirectory=\"%s\" -package -%s -clientconfig=%s -ue4exe=%s %s -utf8output"),
+	FString CommandLine = FString::Printf(TEXT("BuildCookRun %s%s%s%s -nop4 -project=\"%s\" -cook -stage -archive -archivedirectory=\"%s\" -package -platform=%s -clientconfig=%s -ue4exe=%s %s -utf8output"),
 		FRocketSupport::IsRocket() ? TEXT("-rocket ") : TEXT(""),
 		GetUATCompilationFlags(),
 		GetUATNoCodeProjectFlag(),
