[INTEGRATE] Change 2449459 by [REDACTED_PERSON][REDACTED_HANDLE] on 2015/02/17 19:18:52

	Don't stage PDBs if packaging is for distribution. (UE-9983)

diff --git a/Engine/Source/Editor/MainFrame/Private/Frame/MainFrameActions.cpp b/Engine/Source/Editor/MainFrame/Private/Frame/MainFrameActions.cpp
--- a/Engine/Source/Editor/MainFrame/Private/Frame/MainFrameActions.cpp
+++ b/Engine/Source/Editor/MainFrame/Private/Frame/MainFrameActions.cpp
@@ -532,7 +532,7 @@ void FMainFrameActionCallbacks::PackageProject( const FName InPlatformInfoName )
 
 	if (PackagingSettings->ForDistribution)
 	{
-		OptionalParams += TEXT(" -distribution");
+		OptionalParams += TEXT(" -distribution -nodebuginfo");
 	}
 
 	if (PlatformInfo->TargetPlatformName == FName("WindowsNoEditor") && PlatformInfo->PlatformFlavor == TEXT("Win32"))
