Added "Generate project" menu option to File menu, if a VS/XCode project file is not found.

diff --git a/Engine/Source/Editor/MainFrame/Private/Menus/MainMenu.cpp b/Engine/Source/Editor/MainFrame/Private/Menus/MainMenu.cpp
--- a/Engine/Source/Editor/MainFrame/Private/Menus/MainMenu.cpp
+++ b/Engine/Source/Editor/MainFrame/Private/Menus/MainMenu.cpp
@@ -408,6 +408,14 @@ TSharedRef< SWidget > FMainMenu::MakeMainTabMenu( const TSharedPtr<FTabManager>&
 						FText::Format(LOCTEXT("OpenIDETooltip", "Opens your C++ code in {0}."), ShortIDEName)
 					);
 				}
+				else
+				{
+					MenuBuilder.AddMenuEntry( FMainFrameCommands::Get().RefreshCodeProject,
+						NAME_None,
+						FText::Format(LOCTEXT("GenerateCodeProjectLabel", "Generate {0} Project"), ShortIDEName),
+						FText::Format(LOCTEXT("GenerateCodeProjectTooltip", "Generates your C++ code project in {0}."), ShortIDEName)
+					);
+				}
 
 				// @hack GDC: this should be moved somewhere else and be less hacky
 				ITargetPlatform* RunningTargetPlatform = GetTargetPlatformManager()->GetRunningTargetPlatform();
