Made the context menu naming for adding new C++ classes consistent

UE-9418 - Unify description for adding C++ Class to project

diff --git a/Engine/Source/Editor/GameProjectGeneration/Private/GameProjectUtils.cpp b/Engine/Source/Editor/GameProjectGeneration/Private/GameProjectUtils.cpp
--- a/Engine/Source/Editor/GameProjectGeneration/Private/GameProjectUtils.cpp
+++ b/Engine/Source/Editor/GameProjectGeneration/Private/GameProjectUtils.cpp
@@ -701,7 +701,7 @@ void GameProjectUtils::OpenAddToProjectDialog(const FAddToProjectConfig& Config,
 	FText WindowTitle = Config._WindowTitle;
 	if (WindowTitle.IsEmpty())
 	{
-		WindowTitle = InDomain == EClassDomain::Native ? LOCTEXT("AddCodeWindowHeader_Native", "Add Code") : LOCTEXT("AddCodeWindowHeader_Blueprint", "Add Blueprint");
+		WindowTitle = InDomain == EClassDomain::Native ? LOCTEXT("AddCodeWindowHeader_Native", "Add C++ Class") : LOCTEXT("AddCodeWindowHeader_Blueprint", "Add Blueprint Class");
 	}
 
 	TSharedRef<SWindow> AddCodeWindow =
diff --git a/Engine/Source/Editor/MainFrame/Private/Frame/MainFrameActions.cpp b/Engine/Source/Editor/MainFrame/Private/Frame/MainFrameActions.cpp
--- a/Engine/Source/Editor/MainFrame/Private/Frame/MainFrameActions.cpp
+++ b/Engine/Source/Editor/MainFrame/Private/Frame/MainFrameActions.cpp
@@ -62,7 +62,7 @@ void FMainFrameCommands::RegisterCommands()
 	UI_COMMAND( OpenProject, "Open Project...", "Opens a dialog to choose a game project to open", EUserInterfaceActionType::Button, FInputGesture() );
 	ActionList->MapAction( OpenProject, FExecuteAction::CreateStatic( &FMainFrameActionCallbacks::NewProject, true, false), DefaultExecuteAction );
 
-	UI_COMMAND( AddCodeToProject, "Add Code to Project...", "Adds C++ code to the project. The code can only be compiled if you have an appropriate C++ compiler installed.", EUserInterfaceActionType::Button, FInputGesture() );
+	UI_COMMAND( AddCodeToProject, "New C++ Class...", "Adds C++ code to the project. The code can only be compiled if you have an appropriate C++ compiler installed.", EUserInterfaceActionType::Button, FInputGesture() );
 	ActionList->MapAction( AddCodeToProject, FExecuteAction::CreateStatic( &FMainFrameActionCallbacks::AddCodeToProject ));
 
 	UI_COMMAND( RefreshCodeProject, "Refresh code project", "Refreshes your C++ code project.", EUserInterfaceActionType::Button, FInputGesture() );
