found another place where we weren't checking for the ManifestNotFound validation error
#ios

diff --git a/Engine/Source/Editor/MainFrame/Private/Frame/MainFrameActions.cpp b/Engine/Source/Editor/MainFrame/Private/Frame/MainFrameActions.cpp
--- a/Engine/Source/Editor/MainFrame/Private/Frame/MainFrameActions.cpp
+++ b/Engine/Source/Editor/MainFrame/Private/Frame/MainFrameActions.cpp
@@ -432,6 +432,15 @@ void FMainFrameActionCallbacks::PackageProject( const FName InPlatformInfoName )
 				);
 			}
 
+			if ((Result & ETargetPlatformReadyStatus::ManifestNotFound) != 0)
+			{
+				AddMessageLog(
+					LOCTEXT("ManifestNotFound", "Manifest not found."),
+					LOCTEXT("ManifestNotFoundMessageDetail", "The generated application manifest could not be found."),
+					NotInstalledTutorialLink
+					);
+			}
+
 			// report to main frame
 			bool UnrecoverableError = false;
 
