Added new constructor to FRootWindowLocation taking just a window size, and added new public method to determine the best screen location based on the window size and the current display metrics.
Made the Project Browser window a bit higher so it's clear that there's description text below the template image.

diff --git a/Engine/Source/Editor/MainFrame/Private/Frame/RootWindowLocation.h b/Engine/Source/Editor/MainFrame/Private/Frame/RootWindowLocation.h
--- a/Engine/Source/Editor/MainFrame/Private/Frame/RootWindowLocation.h
+++ b/Engine/Source/Editor/MainFrame/Private/Frame/RootWindowLocation.h
@@ -32,29 +32,45 @@ public:
 		: WindowSize( VectorFromSettings(TEXT("WindowSize"), FVector2D(1280,720)) )
 		, InitiallyMaximized( BoolFromSettings(TEXT("InitiallyMaximized"), true) )
 	{
-		// Find the default centered screen position
-		FDisplayMetrics DisplayMetrics;
-		FSlateApplication::Get().GetDisplayMetrics(DisplayMetrics);
-		const FVector2D DisplayTopLeft(DisplayMetrics.PrimaryDisplayWorkAreaRect.Left, DisplayMetrics.PrimaryDisplayWorkAreaRect.Top);
-		const FVector2D DisplaySize(DisplayMetrics.PrimaryDisplayWorkAreaRect.Right - DisplayMetrics.PrimaryDisplayWorkAreaRect.Left, 
-									DisplayMetrics.PrimaryDisplayWorkAreaRect.Bottom - DisplayMetrics.PrimaryDisplayWorkAreaRect.Top);
-		FVector2D DefaultPosition = DisplayTopLeft + (DisplaySize - WindowSize) * 0.5f;
-
-		ScreenPosition = VectorFromSettings(TEXT("ScreenPosition"), DefaultPosition);
+		ScreenPosition = VectorFromSettings(TEXT("ScreenPosition"), GetCenteredScreenPosition() );
 	}
 	
 	/**
-	 * Creates and initializes a new instance.
+	 * Creates and initializes a new instance with the specified size.
+	 */
+	FRootWindowLocation(FVector2D InWindowSize, bool InInitiallyMaximized)
+		: WindowSize(InWindowSize)
+		, InitiallyMaximized(InInitiallyMaximized)
+	{
+		ScreenPosition = GetCenteredScreenPosition();
+	}
+
+	/**
+	 * Creates and initializes a new instance with the specified position and size.
 	 */
 	FRootWindowLocation( FVector2D InScreenPosition, FVector2D InWindowSize, bool InInitiallyMaximized )
-		: ScreenPosition(InScreenPosition)
+		: ScreenPosition( InScreenPosition )
 		, WindowSize( InWindowSize )
-		, InitiallyMaximized(InInitiallyMaximized)
+		, InitiallyMaximized( InInitiallyMaximized )
 	{ }
 
 
 public:
 
+	/**
+	 * Set centered screen position based on the size
+	 */
+	FVector2D GetCenteredScreenPosition() const
+	{
+		// Find the default centered screen position
+		FDisplayMetrics DisplayMetrics;
+		FSlateApplication::Get().GetDisplayMetrics(DisplayMetrics);
+		const FVector2D DisplayTopLeft(DisplayMetrics.PrimaryDisplayWorkAreaRect.Left, DisplayMetrics.PrimaryDisplayWorkAreaRect.Top);
+		const FVector2D DisplaySize(DisplayMetrics.PrimaryDisplayWorkAreaRect.Right - DisplayMetrics.PrimaryDisplayWorkAreaRect.Left,
+									DisplayMetrics.PrimaryDisplayWorkAreaRect.Bottom - DisplayMetrics.PrimaryDisplayWorkAreaRect.Top);
+		return DisplayTopLeft + (DisplaySize - WindowSize) * 0.5f;
+	}
+
 	/**
 	 * Saves this structure to the INI file.
 	 */
diff --git a/Engine/Source/Editor/MainFrame/Private/MainFrameModule.cpp b/Engine/Source/Editor/MainFrame/Private/MainFrameModule.cpp
--- a/Engine/Source/Editor/MainFrame/Private/MainFrameModule.cpp
+++ b/Engine/Source/Editor/MainFrame/Private/MainFrameModule.cpp
@@ -62,11 +62,9 @@ void FMainFrameModule::CreateDefaultMainFrame( const bool bStartImmersivePIE )
 			DefaultWindowLocation.InitiallyMaximized = false;
 
 			DefaultWindowLocation.WindowSize = GetProjectBrowserWindowSize();
+			DefaultWindowLocation.ScreenPosition = DefaultWindowLocation.GetCenteredScreenPosition();
 
-			// Do not let the user adjust the fixed size
 			bIsUserSizable = true;
-
-			// Since this will appear to be a dialog, there is no need to allow maximizing or minimizing
 			bSupportsMaximize = true;
 			bSupportsMinimize = true;
 
diff --git a/Engine/Source/Editor/MainFrame/Private/MainFrameModule.h b/Engine/Source/Editor/MainFrame/Private/MainFrameModule.h
--- a/Engine/Source/Editor/MainFrame/Private/MainFrameModule.h
+++ b/Engine/Source/Editor/MainFrame/Private/MainFrameModule.h
@@ -125,7 +125,7 @@ protected:
 public:
 
 	/** Get the size of the project browser window */
-	static FVector2D GetProjectBrowserWindowSize() { return FVector2D(1100, 730); }
+	static FVector2D GetProjectBrowserWindowSize() { return FVector2D(1100, 780); }
 
 private:
 
