Icons for Project Settings and Editor Preferences tabs

diff --git a/Engine/Source/Editor/EditorSettingsViewer/EditorSettingsViewer.Build.cs b/Engine/Source/Editor/EditorSettingsViewer/EditorSettingsViewer.Build.cs
--- a/Engine/Source/Editor/EditorSettingsViewer/EditorSettingsViewer.Build.cs
+++ b/Engine/Source/Editor/EditorSettingsViewer/EditorSettingsViewer.Build.cs
@@ -24,7 +24,8 @@ namespace UnrealBuildTool.Rules
 					"SlateCore",
 					"UnrealEd",
                     "InternationalizationSettings",
-					"BlueprintGraph"
+					"BlueprintGraph",
+                    "EditorStyle",
 				}
 			);
 
diff --git a/Engine/Source/Editor/EditorSettingsViewer/Private/EditorSettingsViewerModule.cpp b/Engine/Source/Editor/EditorSettingsViewer/Private/EditorSettingsViewerModule.cpp
--- a/Engine/Source/Editor/EditorSettingsViewer/Private/EditorSettingsViewerModule.cpp
+++ b/Engine/Source/Editor/EditorSettingsViewer/Private/EditorSettingsViewerModule.cpp
@@ -66,7 +66,8 @@ public:
 
 		FGlobalTabmanager::Get()->RegisterNomadTabSpawner(EditorSettingsTabName, FOnSpawnTab::CreateRaw(this, &FEditorSettingsViewerModule::HandleSpawnSettingsTab))
 			.SetDisplayName(LOCTEXT("EditorSettingsTabTitle", "Editor Preferences"))
-			.SetMenuType(ETabSpawnerMenuType::Hidden);
+			.SetMenuType(ETabSpawnerMenuType::Hidden)
+			.SetIcon(FSlateIcon(FEditorStyle::GetStyleSetName(), "EditorPreferences.TabIcon"));
 	}
 
 	virtual void ShutdownModule() override
diff --git a/Engine/Source/Editor/EditorStyle/Private/SlateEditorStyle.cpp b/Engine/Source/Editor/EditorStyle/Private/SlateEditorStyle.cpp
--- a/Engine/Source/Editor/EditorStyle/Private/SlateEditorStyle.cpp
+++ b/Engine/Source/Editor/EditorStyle/Private/SlateEditorStyle.cpp
@@ -1401,14 +1401,21 @@ void FSlateEditorStyle::FStyle::SetupGeneralStyles()
 		Set( "MainFrame.RecentLevels", new IMAGE_BRUSH( "Icons/icon_file_LevelsRecent_16px", Icon16x16 ) );
 		Set( "MainFrame.FavoriteLevels", new IMAGE_BRUSH( "Icons/icon_file_LevelsRecent_16px", Icon16x16 ) );
 
-		Set( "MainFrame.EditorPreferences", new IMAGE_BRUSH( "Icons/Edit/icon_Edit_EditorPreferences_16x", Icon16x16 ) );
-		Set( "MainFrame.ProjectSettings", new IMAGE_BRUSH( "Icons/Edit/icon_Edit_ProjectSettings_16x", Icon16x16 ) );
-
 		Set( "MainFrame.DebugTools.SmallFont", TTF_CORE_FONT( "Fonts/Roboto-Regular", 8 ) );
 		Set( "MainFrame.DebugTools.NormalFont", TTF_CORE_FONT( "Fonts/Roboto-Regular", 9 ) );
 		Set( "MainFrame.DebugTools.LabelFont", TTF_CORE_FONT( "Fonts/Roboto-Regular", 8 ) );
 	}
 
+	// Editor preferences
+	{
+		Set("EditorPreferences.TabIcon", new IMAGE_BRUSH("Icons/Edit/icon_Edit_EditorPreferences_16x", Icon16x16));
+	}
+
+	// Project settings
+	{
+		Set("ProjectSettings.TabIcon", new IMAGE_BRUSH("Icons/Edit/icon_Edit_ProjectSettings_16x", Icon16x16));
+	}
+
 	// Main frame
 	{
 		Set("MainFrame.StatusInfoButton", FButtonStyle(Button)
diff --git a/Engine/Source/Editor/MainFrame/Private/Menus/MainMenu.cpp b/Engine/Source/Editor/MainFrame/Private/Menus/MainMenu.cpp
--- a/Engine/Source/Editor/MainFrame/Private/Menus/MainMenu.cpp
+++ b/Engine/Source/Editor/MainFrame/Private/Menus/MainMenu.cpp
@@ -120,13 +120,17 @@ void FMainMenu::FillEditMenu( FMenuBuilder& MenuBuilder, const TSharedRef< FExte
 			MenuBuilder.AddSubMenu(
 				LOCTEXT("EditorPreferencesSubMenuLabel", "Editor Preferences"),
 				LOCTEXT("EditorPreferencesSubMenuToolTip", "Configure the behavior and features of this Editor"),
-				FNewMenuDelegate::CreateStatic(&FSettingsMenu::MakeMenu, FName("Editor"))
+				FNewMenuDelegate::CreateStatic(&FSettingsMenu::MakeMenu, FName("Editor")),
+				false,
+				FSlateIcon(FEditorStyle::GetStyleSetName(), "EditorPreferences.TabIcon")
 			);
 
 			MenuBuilder.AddSubMenu(
 				LOCTEXT("ProjectSettingsSubMenuLabel", "Project Settings"),
 				LOCTEXT("ProjectSettingsSubMenuToolTip", "Change the settings of the currently loaded project"),
-				FNewMenuDelegate::CreateStatic(&FSettingsMenu::MakeMenu, FName("Project"))
+				FNewMenuDelegate::CreateStatic(&FSettingsMenu::MakeMenu, FName("Project")),
+				false,
+				FSlateIcon(FEditorStyle::GetStyleSetName(), "ProjectSettings.TabIcon")
 			);
 		}
 		else
@@ -135,7 +139,7 @@ void FMainMenu::FillEditMenu( FMenuBuilder& MenuBuilder, const TSharedRef< FExte
 			MenuBuilder.AddMenuEntry(
 				LOCTEXT("EditorPreferencesMenuLabel", "Editor Preferences..."),
 				LOCTEXT("EditorPreferencesMenuToolTip", "Configure the behavior and features of the Unreal Editor."),
-				FSlateIcon(),
+				FSlateIcon(FEditorStyle::GetStyleSetName(), "EditorPreferences.TabIcon"),
 				FUIAction(FExecuteAction::CreateStatic(&FSettingsMenu::OpenSettings, FName("Editor"), FName("General"), FName("Appearance")))
 			);
 #endif
@@ -143,7 +147,7 @@ void FMainMenu::FillEditMenu( FMenuBuilder& MenuBuilder, const TSharedRef< FExte
 			MenuBuilder.AddMenuEntry(
 				LOCTEXT("ProjectSettingsMenuLabel", "Project Settings..."),
 				LOCTEXT("ProjectSettingsMenuToolTip", "Change the settings of the currently loaded project."),
-				FSlateIcon(),
+				FSlateIcon(FEditorStyle::GetStyleSetName(), "ProjectSettings.TabIcon"),
 				FUIAction(FExecuteAction::CreateStatic(&FSettingsMenu::OpenSettings, FName("Project"), FName("Project"), FName("General")))
 			);
 		}
diff --git a/Engine/Source/Editor/ProjectSettingsViewer/Private/ProjectSettingsViewerModule.cpp b/Engine/Source/Editor/ProjectSettingsViewer/Private/ProjectSettingsViewerModule.cpp
--- a/Engine/Source/Editor/ProjectSettingsViewer/Private/ProjectSettingsViewerModule.cpp
+++ b/Engine/Source/Editor/ProjectSettingsViewer/Private/ProjectSettingsViewerModule.cpp
@@ -86,7 +86,8 @@ public:
 
 		FGlobalTabmanager::Get()->RegisterNomadTabSpawner(ProjectSettingsTabName, FOnSpawnTab::CreateRaw(this, &FProjectSettingsViewerModule::HandleSpawnSettingsTab))
 			.SetDisplayName(LOCTEXT("ProjectSettingsTabTitle", "Project Settings"))
-			.SetMenuType(ETabSpawnerMenuType::Hidden);
+			.SetMenuType(ETabSpawnerMenuType::Hidden)
+			.SetIcon(FSlateIcon(FEditorStyle::GetStyleSetName(), "ProjectSettings.TabIcon"));
 	}
 
 	virtual void ShutdownModule() override
diff --git a/Engine/Source/Editor/ProjectSettingsViewer/ProjectSettingsViewer.Build.cs b/Engine/Source/Editor/ProjectSettingsViewer/ProjectSettingsViewer.Build.cs
--- a/Engine/Source/Editor/ProjectSettingsViewer/ProjectSettingsViewer.Build.cs
+++ b/Engine/Source/Editor/ProjectSettingsViewer/ProjectSettingsViewer.Build.cs
@@ -26,6 +26,7 @@ namespace UnrealBuildTool.Rules
 					"MoviePlayer",
                     "AIModule",
 					"ProjectTargetPlatformEditor",
+                    "EditorStyle",
 				}
 			);
 
