Force UAT to build content-only projects as well to avoid issues with missing/outdated UE4Game binaries.

diff --git a/Engine/Source/Editor/MainFrame/Private/Frame/MainFrameActions.cpp b/Engine/Source/Editor/MainFrame/Private/Frame/MainFrameActions.cpp
--- a/Engine/Source/Editor/MainFrame/Private/Frame/MainFrameActions.cpp
+++ b/Engine/Source/Editor/MainFrame/Private/Frame/MainFrameActions.cpp
@@ -566,7 +566,7 @@ void FMainFrameActionCallbacks::PackageProject( const FName InPlatformInfoName )
 	}
 
 	// only build if the project has code that might need to be built
-	if (bProjectHasCode && FSourceCodeNavigation::IsCompilerAvailable())
+	if (FSourceCodeNavigation::IsCompilerAvailable())
 	{
 		OptionalParams += TEXT(" -build");
 	}
