UE-3133 - "Binary: Translation Editor: Missing menu items when enabled"

Deprecating Translation Editor menu, use Localization Dashboard instead.
New experimental setting for enabling just the translation picker.

diff --git a/Engine/Source/Editor/MainFrame/Private/Menus/MainMenu.cpp b/Engine/Source/Editor/MainFrame/Private/Menus/MainMenu.cpp
--- a/Engine/Source/Editor/MainFrame/Private/Menus/MainMenu.cpp
+++ b/Engine/Source/Editor/MainFrame/Private/Menus/MainMenu.cpp
@@ -195,11 +195,12 @@ void FMainMenu::FillWindowMenu( FMenuBuilder& MenuBuilder, const TSharedRef< FEx
 		bool bMessagingDebugger = GetDefault<UEditorExperimentalSettings>()->bMessagingDebugger;
 		bool bBlutility = GetDefault<UEditorExperimentalSettings>()->bEnableEditorUtilityBlueprints;
 		bool bLocalizationDashboard = GetDefault<UEditorExperimentalSettings>()->bEnableLocalizationDashboard;
+		bool bTranslationPicker = GetDefault<UEditorExperimentalSettings>()->bEnableTranslationPicker;
 		bool bTranslationEditor = GetDefault<UEditorExperimentalSettings>()->bEnableTranslationEditor;
 		bool bMergeActors = GetDefault<UEditorExperimentalSettings>()->bActorMerging;
 
 		// Make sure at least one is enabled before creating the section
-		if (bMessagingDebugger || bBlutility || bLocalizationDashboard || bTranslationEditor || bMergeActors)
+		if (bMessagingDebugger || bBlutility || bLocalizationDashboard || bTranslationEditor || bTranslationPicker || bMergeActors)
 		{
 			MenuBuilder.BeginSection("ExperimentalTabSpawners", LOCTEXT("ExperimentalTabSpawnersHeading", "Experimental"));
 			{
@@ -239,11 +240,21 @@ void FMainMenu::FillWindowMenu( FMenuBuilder& MenuBuilder, const TSharedRef< FEx
 				// Translation Editor
 				if (bTranslationEditor)
 				{
+					MenuBuilder.AddMenuEntry(
+						LOCTEXT("TranslationEditorMenuDeprecated", "Please use Localization Dashboard"),
+						LOCTEXT("TranslationEditorMenuDeprecatedToolTip", "The Translation Editor is now accessible from the Localization Dashboard (experimental). \nPlease access it from there, as the Translation Editor menu below will be removed in a future update."),
+						FSlateIcon(),
+						FUIAction()
+						);
+
 					MenuBuilder.AddSubMenu(
 						LOCTEXT("TranslationEditorSubMenuLabel", "Translation Editor"),
-						LOCTEXT("EditorPreferencesSubMenuToolTip", "Open the Translation Editor for a Given Project and Language"),
+						LOCTEXT("TranslationEditorSubMenuToolTip", "Open the Translation Editor for a Given Project and Language"),
 						FNewMenuDelegate::CreateStatic(&FMainFrameTranslationEditorMenu::MakeMainFrameTranslationEditorSubMenu)
 						);
+				}
+				if (bTranslationPicker)
+				{
 					MenuBuilder.AddMenuEntry(
 						LOCTEXT("TranslationPickerMenuItem", "Translation Picker"),
 						LOCTEXT("TranslationPickerMenuItemToolTip", "Launch the Translation Picker to Modify Editor Translations"),
diff --git a/Engine/Source/Editor/UnrealEd/Classes/Settings/EditorExperimentalSettings.h b/Engine/Source/Editor/UnrealEd/Classes/Settings/EditorExperimentalSettings.h
--- a/Engine/Source/Editor/UnrealEd/Classes/Settings/EditorExperimentalSettings.h
+++ b/Engine/Source/Editor/UnrealEd/Classes/Settings/EditorExperimentalSettings.h
@@ -29,10 +29,15 @@ public:
 	UPROPERTY(EditAnywhere, config, Category = Tools, meta = (DisplayName = "Localization Dashboard"))
 	bool bEnableLocalizationDashboard;
 
-	/** Allows usage of the Translation Editor */
-	UPROPERTY(EditAnywhere, config, Category = Tools, meta = (DisplayName = "Translation Editor"))
+	/** (DEPRECATED) Allows usage of the Translation Editor */
+	DEPRECATED(4.8, "The Translation Editor menu will be removed. The Translation Editor is now accessed from the Localization Dashboard.")
+	UPROPERTY(EditAnywhere, config, Category = Tools, meta = (DisplayName = "Translation Editor (Menu deprecated, please access from Localization Dashboard)"))
 	bool bEnableTranslationEditor;
 
+	/** Allows usage of the Translation Picker */
+	UPROPERTY(EditAnywhere, config, Category = Tools, meta = (DisplayName = "Translation Picker"))
+	bool bEnableTranslationPicker;
+
 	/** The Blutility shelf holds editor utility Blueprints. Summon from the Workspace menu. */
 	UPROPERTY(EditAnywhere, config, Category=Tools, meta=(DisplayName="Editor Utility Blueprints (Blutility)"))
 	bool bEnableEditorUtilityBlueprints;
