Fix for CrashReporter being a dependency when packaging windows XP.

diff --git a/Engine/Source/Editor/MainFrame/Private/Frame/MainFrameActions.cpp b/Engine/Source/Editor/MainFrame/Private/Frame/MainFrameActions.cpp
--- a/Engine/Source/Editor/MainFrame/Private/Frame/MainFrameActions.cpp
+++ b/Engine/Source/Editor/MainFrame/Private/Frame/MainFrameActions.cpp
@@ -541,6 +541,8 @@ void FMainFrameActionCallbacks::PackageProject( const FName InPlatformInfoName )
 		OptionalParams += TEXT(" -nodebuginfo");
 	}
 
+
+	bool bTargetPlatformCanUseCrashReporter = true;
 	if (PlatformInfo->TargetPlatformName == FName("WindowsNoEditor") && PlatformInfo->PlatformFlavor == TEXT("Win32"))
 	{
 		FString MinumumSupportedWindowsOS;
@@ -548,6 +550,7 @@ void FMainFrameActionCallbacks::PackageProject( const FName InPlatformInfoName )
 		if (MinumumSupportedWindowsOS == TEXT("MSOS_XP"))
 		{
 			OptionalParams += TEXT(" -OverrideMinimumOS=WinXP");
+			bTargetPlatformCanUseCrashReporter = false;
 		}
 	}
 
@@ -570,7 +573,7 @@ void FMainFrameActionCallbacks::PackageProject( const FName InPlatformInfoName )
 	}
 
 	// we can include the crash reporter for any non-code projects or internal builds (we don't have the symbols to interpret external builds of code-based projects)
-	if (PackagingSettings->IncludeCrashReporter && (!bProjectHasCode || FEngineBuildSettings::IsInternalBuild()))
+	if (PackagingSettings->IncludeCrashReporter && (!bProjectHasCode || FEngineBuildSettings::IsInternalBuild()) && bTargetPlatformCanUseCrashReporter)
 	{
 		OptionalParams += TEXT(" -CrashReporter");
 	}
