Remove Translation Editor experimental feature option and menu.  Now accessed through the localization dashboard (which is also experimental).

diff --git a/Engine/Source/Editor/MainFrame/Private/Menus/MainMenu.cpp b/Engine/Source/Editor/MainFrame/Private/Menus/MainMenu.cpp
--- a/Engine/Source/Editor/MainFrame/Private/Menus/MainMenu.cpp
+++ b/Engine/Source/Editor/MainFrame/Private/Menus/MainMenu.cpp
@@ -196,11 +196,10 @@ void FMainMenu::FillWindowMenu( FMenuBuilder& MenuBuilder, const TSharedRef< FEx
 		bool bBlutility = GetDefault<UEditorExperimentalSettings>()->bEnableEditorUtilityBlueprints;
 		bool bLocalizationDashboard = GetDefault<UEditorExperimentalSettings>()->bEnableLocalizationDashboard;
 		bool bTranslationPicker = GetDefault<UEditorExperimentalSettings>()->bEnableTranslationPicker;
-		bool bTranslationEditor = GetDefault<UEditorExperimentalSettings>()->bEnableTranslationEditor;
 		bool bMergeActors = GetDefault<UEditorExperimentalSettings>()->bActorMerging;
 
 		// Make sure at least one is enabled before creating the section
-		if (bMessagingDebugger || bBlutility || bLocalizationDashboard || bTranslationEditor || bTranslationPicker || bMergeActors)
+		if (bMessagingDebugger || bBlutility || bLocalizationDashboard || bTranslationPicker || bMergeActors)
 		{
 			MenuBuilder.BeginSection("ExperimentalTabSpawners", LOCTEXT("ExperimentalTabSpawnersHeading", "Experimental"));
 			{
@@ -237,22 +236,7 @@ void FMainMenu::FillWindowMenu( FMenuBuilder& MenuBuilder, const TSharedRef< FEx
 						);
 				}
 
-				// Translation Editor
-				if (bTranslationEditor)
-				{
-					MenuBuilder.AddMenuEntry(
-						LOCTEXT("TranslationEditorMenuDeprecated", "Please use Localization Dashboard"),
-						LOCTEXT("TranslationEditorMenuDeprecatedToolTip", "The Translation Editor is now accessible from the Localization Dashboard (experimental). \nPlease access it from there, as the Translation Editor menu below will be removed in a future update."),
-						FSlateIcon(),
-						FUIAction()
-						);
-
-					MenuBuilder.AddSubMenu(
-						LOCTEXT("TranslationEditorSubMenuLabel", "Translation Editor"),
-						LOCTEXT("TranslationEditorSubMenuToolTip", "Open the Translation Editor for a Given Project and Language"),
-						FNewMenuDelegate::CreateStatic(&FMainFrameTranslationEditorMenu::MakeMainFrameTranslationEditorSubMenu)
-						);
-				}
+				// Translation Picker
 				if (bTranslationPicker)
 				{
 					MenuBuilder.AddMenuEntry(
diff --git a/Engine/Source/Editor/MainFrame/Private/Menus/TranslationEditorMenu.h b/Engine/Source/Editor/MainFrame/Private/Menus/TranslationEditorMenu.h
--- a/Engine/Source/Editor/MainFrame/Private/Menus/TranslationEditorMenu.h
+++ b/Engine/Source/Editor/MainFrame/Private/Menus/TranslationEditorMenu.h
@@ -13,180 +13,12 @@ class FMainFrameTranslationEditorMenu
 {
 public:
 
-	/**
-	 * Creates the main frame translation editor sub menu.
-	 *
-	 * @param MenuBuilder The builder for the menu that owns this menu.
-	 */
-	static void MakeMainFrameTranslationEditorSubMenu( FMenuBuilder& MenuBuilder )
-	{
-		// Directories to search for localization
-		// Directory Structure should look like this:
-		// ProjectRootFolder/Content/Localization/  <-- This is the directory to add to the array below
-		// ProjectRootFolder/Content/Localization/LocalizationProjectName/
-		// ProjectRootFolder/Content/Localization/LocalizationProjectName/LocalizationProjectName.manifest
-		// ProjectRootFolder/Content/Localization/LocalizationProjectName/en/LocalizationProjectName.archive
-		// ProjectRootFolder/Content/Localization/LocalizationProjectName/jp/LocalizationProjectName.archive
-		// ProjectRootFolder/Content/Localization/LocalizationProjectName/ko/LocalizationProjectName.archive
-		// ...
-
-		// Add some specific directories 
-		TArray<FString> DirectoriesToSearch;
-		DirectoriesToSearch.Add(FPaths::EngineContentDir() / TEXT("Localization"));
-		DirectoriesToSearch.Add(FPaths::GameContentDir() / TEXT("Localization"));
-		DirectoriesToSearch.Add(FPaths::EngineDir() / "Programs/NoRedist/UnrealEngineLauncher/Content/Localization");
-
-		// Also check any of these .ini specified localization paths
-		TArray < TArray<FString> > LocalizationPathArrays;
-		LocalizationPathArrays.Add(FPaths::GetGameLocalizationPaths());
-		LocalizationPathArrays.Add(FPaths::GetEditorLocalizationPaths());
-		LocalizationPathArrays.Add(FPaths::GetEngineLocalizationPaths());
-		LocalizationPathArrays.Add(FPaths::GetPropertyNameLocalizationPaths());
-		LocalizationPathArrays.Add(FPaths::GetToolTipLocalizationPaths());
-		for (TArray<FString> LocalizationPathArray : LocalizationPathArrays)
-		{
-			for (FString LocalizationPath : LocalizationPathArray)
-			{
-				DirectoriesToSearch.AddUnique(LocalizationPath);
-			}
-		}
-
-		// Find Manifest files
-		TArray<FString> ManifestFileNames;
-		for (FString DirectoryToSearch : DirectoriesToSearch)
-		{
-			// Search for .manifest files in this directory
-			const FString ManifestExtension = TEXT("manifest");
-			const FString TopFolderManifestFileWildcard = FString::Printf(TEXT("%s/*.%s"), *DirectoryToSearch, *ManifestExtension);
-
-			TArray<FString> TopFolderManifestFileNames;
-			IFileManager::Get().FindFiles(TopFolderManifestFileNames, *TopFolderManifestFileWildcard, true, false);
-
-			for (FString& CurrentTopFolderManifestFileName : TopFolderManifestFileNames)
-			{
-				// Add root back on
-				CurrentTopFolderManifestFileName = DirectoryToSearch / CurrentTopFolderManifestFileName;
-				ManifestFileNames.AddUnique(CurrentTopFolderManifestFileName);
-			}
-
-			// Search for .manifest files 1 directory deep
-			const FString DirectoryToSearchWildcard = FString::Printf(TEXT("%s/*"), *DirectoryToSearch);
-			TArray<FString> ProjectFolders;
-			IFileManager::Get().FindFiles(ProjectFolders, *DirectoryToSearchWildcard, false, true);
-
-			for (FString ProjectFolder : ProjectFolders)
-			{
-				// Add root back on
-				ProjectFolder = DirectoryToSearch / ProjectFolder;
-				const FString ProjectFolderManifestFileWildcard = FString::Printf(TEXT("%s/*.%s"), *ProjectFolder, *ManifestExtension);
-
-				TArray<FString> CurrentFolderManifestFileNames;
-				IFileManager::Get().FindFiles(CurrentFolderManifestFileNames, *ProjectFolderManifestFileWildcard, true, false);
-
-				for (FString& CurrentFolderManifestFileName : CurrentFolderManifestFileNames)
-				{
-					// Add root back on
-					CurrentFolderManifestFileName = ProjectFolder / CurrentFolderManifestFileName;
-					ManifestFileNames.AddUnique(CurrentFolderManifestFileName);
-				}
-			}
-		}
-
-		MenuBuilder.BeginSection("Project", LOCTEXT("Translation_ChooseProject", "Choose Project:"));
-		{
-			for (FString ManifestFileName : ManifestFileNames)
-			{
-				FString ManifestName = FPaths::GetBaseFilename(ManifestFileName);
-
-				FFormatNamedArguments Arguments;
-				Arguments.Add(TEXT("ManifestName"), FText::FromString(ManifestName));
-				MenuBuilder.AddSubMenu(
-					FText::Format(LOCTEXT("TranslationEditorSubMenuProjectLabel", "{ManifestName} Translations"), Arguments),
-					FText::Format(LOCTEXT("TranslationEditorSubMenuProjectToolTip", "Edit Translations for {ManifestName} Strings"), Arguments),
-					FNewMenuDelegate::CreateStatic(&FMainFrameTranslationEditorMenu::MakeMainFrameTranslationEditorSubMenuEditorProject, ManifestFileName)
-				);
-			}
-		}
-
-		MenuBuilder.EndSection();
-	}
-
-	/**
-	 * Creates the main frame translation editor sub menu for editor translations.
-	 *
-	 * @param MenuBuilder The builder for the menu that owns this menu.
-	 */
-	static void MakeMainFrameTranslationEditorSubMenuEditorProject( FMenuBuilder& MenuBuilder, const FString ManifestFileName )
-	{
-		const FString ManifestFolder = FPaths::GetPath(ManifestFileName);
-		const FString ManifestFolderWildcard = FString::Printf(TEXT("%s/*"), *ManifestFolder);
-
-		TArray<FString> ArchiveFolders;
-		IFileManager::Get().FindFiles(ArchiveFolders, *ManifestFolderWildcard, false, true);
-
-		TArray<FString> ArchiveFileNames;
-		for (FString ArchiveFolder : ArchiveFolders)
-		{
-			// Add root back on
-			ArchiveFolder = ManifestFolder / ArchiveFolder;
-
-			const FString ArchiveExtension = TEXT("archive");
-			const FString ArchiveFileWildcard = FString::Printf(TEXT("%s/*.%s"), *ArchiveFolder, *ArchiveExtension);
-
-			TArray<FString> CurrentFolderArchiveFileNames;
-			IFileManager::Get().FindFiles(CurrentFolderArchiveFileNames, *ArchiveFileWildcard, true, false);
-
-			for (FString& CurrentFolderArchiveFileName : CurrentFolderArchiveFileNames)
-			{
-				// Add root back on
-				CurrentFolderArchiveFileName = ArchiveFolder / CurrentFolderArchiveFileName;
-			}
-
-			ArchiveFileNames.Append(CurrentFolderArchiveFileNames);
-		}
-
-		MenuBuilder.BeginSection("Language", LOCTEXT("Translation_ChooseLanguage", "Choose Language:"));
-		{
-			for (FString ArchiveFileName : ArchiveFileNames)
-			{
-				FString ArchiveName = FPaths::GetBaseFilename(FPaths::GetPath(ArchiveFileName));
-				FString ManifestName = FPaths::GetBaseFilename(ManifestFileName);
-
-				FCulturePtr Culture = FInternationalization::Get().GetCulture(ArchiveName);
-				if (Culture.IsValid())
-				{
-					ArchiveName = Culture->GetDisplayName();
-				}
-
-				FFormatNamedArguments Arguments;
-				Arguments.Add(TEXT("ArchiveName"), FText::FromString(ArchiveName));
-				Arguments.Add(TEXT("ManifestName"), FText::FromString(ManifestName));
-
-				MenuBuilder.AddMenuEntry(
-					FText::Format(LOCTEXT("TranslationEditorSubMenuLanguageLabel", "{ArchiveName} Translations"), Arguments),
-					FText::Format(LOCTEXT("TranslationEditorSubMenuLanguageToolTip", "Edit the {ArchiveName} Translations for the {ManifestName} Project"), Arguments),
-					FSlateIcon(),
-					FUIAction(FExecuteAction::CreateStatic(&FMainFrameTranslationEditorMenu::HandleOpenTranslationEditor, ManifestFileName, ArchiveFileName))
-					);
-			}
-		}
-
-		MenuBuilder.EndSection();
-	}
-
 	// Handles clicking a menu entry.
 	static void HandleOpenTranslationPicker()
 	{
 		FModuleManager::LoadModuleChecked<ITranslationEditor>("TranslationEditor").OpenTranslationPicker();
 	}
 
-private:
-
-	// Handles clicking a menu entry.
-	static void HandleOpenTranslationEditor( FString ProjectName, FString Language)
-	{
-		FModuleManager::LoadModuleChecked<ITranslationEditor>("TranslationEditor").OpenTranslationEditor(ProjectName, Language);
-	}
 };
 
 
diff --git a/Engine/Source/Editor/UnrealEd/Classes/Settings/EditorExperimentalSettings.h b/Engine/Source/Editor/UnrealEd/Classes/Settings/EditorExperimentalSettings.h
--- a/Engine/Source/Editor/UnrealEd/Classes/Settings/EditorExperimentalSettings.h
+++ b/Engine/Source/Editor/UnrealEd/Classes/Settings/EditorExperimentalSettings.h
@@ -29,10 +29,6 @@ public:
 	UPROPERTY(EditAnywhere, config, Category = Tools, meta = (DisplayName = "Localization Dashboard"))
 	bool bEnableLocalizationDashboard;
 
-	/** (DEPRECATED) Allows usage of the Translation Editor */
-	UPROPERTY(EditAnywhere, config, Category = Tools, meta = (DisplayName = "Translation Editor (Menu deprecated, please access from Localization Dashboard)"))
-	bool bEnableTranslationEditor;
-
 	/** Allows usage of the Translation Picker */
 	UPROPERTY(EditAnywhere, config, Category = Tools, meta = (DisplayName = "Translation Picker"))
 	bool bEnableTranslationPicker;
