MainFrame - Add  a low risk hack for 4.8 which doesn't force the window to close immediately in the choose project window to avoid a shut down crash when using opengl in windows.  A proper fix is coming for 4.9.

diff --git a/Engine/Source/Editor/MainFrame/Private/Frame/MainFrameHandler.cpp b/Engine/Source/Editor/MainFrame/Private/Frame/MainFrameHandler.cpp
--- a/Engine/Source/Editor/MainFrame/Private/Frame/MainFrameHandler.cpp
+++ b/Engine/Source/Editor/MainFrame/Private/Frame/MainFrameHandler.cpp
@@ -52,7 +52,13 @@ void FMainFrameHandler::ShutDownEditor()
 	if (RootWindow.IsValid())
 	{
 		RootWindow->SetRequestDestroyWindowOverride(FRequestDestroyWindowOverride());
-		RootWindow->RequestDestroyWindow();
+		
+		// This is a low risk temporary hack for 4.8 to avoid a crash when shutting down the open project window when running
+		// in opengl on windows.  A proper, higher risk fix is coming which will replace this.
+		if (FApp::HasGameName())
+		{
+			RootWindow->RequestDestroyWindow();
+		}
 	}
 
 	// Save out any config settings for the editor so they don't get lost
