Undoing changes in CL# 2537281.

Not compatible with users who have no source, but may have visual studio installed.

diff --git a/Engine/Source/Editor/MainFrame/Private/Frame/MainFrameActions.cpp b/Engine/Source/Editor/MainFrame/Private/Frame/MainFrameActions.cpp
--- a/Engine/Source/Editor/MainFrame/Private/Frame/MainFrameActions.cpp
+++ b/Engine/Source/Editor/MainFrame/Private/Frame/MainFrameActions.cpp
@@ -556,7 +556,7 @@ void FMainFrameActionCallbacks::PackageProject( const FName InPlatformInfoName )
 	}
 
 	// only build if the project has code that might need to be built
-	if (FSourceCodeNavigation::IsCompilerAvailable())
+	if (bProjectHasCode && FSourceCodeNavigation::IsCompilerAvailable())
 	{
 		OptionalParams += TEXT(" -build");
 	}
