Updated help button tooltips, including forums, wiki, documentation, and answer hub.

UE-16486, UE-16487

diff --git a/Engine/Source/Editor/LevelEditor/Private/LevelEditorActions.cpp b/Engine/Source/Editor/LevelEditor/Private/LevelEditorActions.cpp
--- a/Engine/Source/Editor/LevelEditor/Private/LevelEditorActions.cpp
+++ b/Engine/Source/Editor/LevelEditor/Private/LevelEditorActions.cpp
@@ -2864,7 +2864,7 @@ class UWorld* FLevelEditorActionCallbacks::GetWorld()
 PRAGMA_DISABLE_OPTIMIZATION
 void FLevelEditorCommands::RegisterCommands()
 {
-	UI_COMMAND( BrowseDocumentation, "Documentation...", "Opens the main documentation page", EUserInterfaceActionType::Button, FInputChord( EKeys::F1 ) );
+	UI_COMMAND( BrowseDocumentation, "Documentation...", "Opens the main documentation page, and allows you to search across all UE4 support sites.", EUserInterfaceActionType::Button, FInputChord( EKeys::F1 ) );
 	UI_COMMAND( BrowseAPIReference, "API Reference...", "Opens the API reference documentation", EUserInterfaceActionType::Button, FInputChord() );
 	UI_COMMAND( BrowseViewportControls, "Viewport Controls...", "Opens the viewport controls cheat sheet", EUserInterfaceActionType::Button, FInputChord() );
 	UI_COMMAND( NewLevel, "New Level...", "Create a new level, or choose a level template to start from.", EUserInterfaceActionType::Button, FInputChord( EModifierKey::Control, EKeys::N ) );
diff --git a/Engine/Source/Editor/MainFrame/Private/Frame/MainFrameActions.cpp b/Engine/Source/Editor/MainFrame/Private/Frame/MainFrameActions.cpp
--- a/Engine/Source/Editor/MainFrame/Private/Frame/MainFrameActions.cpp
+++ b/Engine/Source/Editor/MainFrame/Private/Frame/MainFrameActions.cpp
@@ -121,16 +121,16 @@ void FMainFrameCommands::RegisterCommands()
 												FCanExecuteAction(),
 												FIsActionChecked::CreateStatic( &FMainFrameActionCallbacks::OpenSlateApp_IsChecked, FName("SessionFrontend" ) ) );
 
-	UI_COMMAND(VisitWiki, "Wiki...", "Go to the Unreal Engine Wiki page", EUserInterfaceActionType::Button, FInputChord());
+	UI_COMMAND(VisitWiki, "Wiki...", "Go to the Unreal Engine Wiki page to view community-created resources, or to create your own.", EUserInterfaceActionType::Button, FInputChord());
 	ActionList->MapAction(VisitWiki, FExecuteAction::CreateStatic(&FMainFrameActionCallbacks::VisitWiki));
 
-	UI_COMMAND(VisitForums, "Forums...", "Go to the Unreal Engine forums", EUserInterfaceActionType::Button, FInputChord());
+	UI_COMMAND(VisitForums, "Forums...", "Go the the Unreal Engine forums to view announcements and engage in discussions with other developers.", EUserInterfaceActionType::Button, FInputChord());
 	ActionList->MapAction(VisitForums, FExecuteAction::CreateStatic(&FMainFrameActionCallbacks::VisitForums));
 
 	UI_COMMAND( VisitAskAQuestionPage, "Ask a Question...", "Have a question?  Go here to ask about anything and everything related to Unreal.", EUserInterfaceActionType::Button, FInputChord() );
 	ActionList->MapAction( VisitAskAQuestionPage, FExecuteAction::CreateStatic( &FMainFrameActionCallbacks::VisitAskAQuestionPage ) );
 
-	UI_COMMAND( VisitSearchForAnswersPage, "Answer Hub...", "Searches for useful answers on UDN provided by other users and experts.", EUserInterfaceActionType::Button, FInputChord() );
+	UI_COMMAND( VisitSearchForAnswersPage, "Answer Hub...", "Go to the AnswerHub to ask questions, search existing answers, and share your knowledge with other UE4 developers.", EUserInterfaceActionType::Button, FInputChord() );
 	ActionList->MapAction( VisitSearchForAnswersPage, FExecuteAction::CreateStatic( &FMainFrameActionCallbacks::VisitSearchForAnswersPage ) );
 
 	UI_COMMAND( VisitSupportWebSite, "Unreal Engine Support Web Site...", "Navigates to the Unreal Engine Support web site's main page.", EUserInterfaceActionType::Button, FInputChord() );
