Package project Cook and package step errors are now presented to the Message Log

UEPLAT-880

diff --git a/Engine/Source/Editor/MainFrame/Private/Frame/MainFrameActions.cpp b/Engine/Source/Editor/MainFrame/Private/Frame/MainFrameActions.cpp
--- a/Engine/Source/Editor/MainFrame/Private/Frame/MainFrameActions.cpp
+++ b/Engine/Source/Editor/MainFrame/Private/Frame/MainFrameActions.cpp
@@ -1158,6 +1158,90 @@ void FMainFrameActionCallbacks::HandleUatProcessCanceled( TWeakPtr<SNotification
 }
 
 
+/**
+ * Helper class to deal with packaging issues encountered in UAT.
+ **/
+class FPackagingErrorHandler
+{
+private:
+
+	/**
+	 * Create a message to send to the Message Log.
+	 *
+	 * @Param MessageString - The error we wish to send to the Message Log.
+	 * @Param MessageType - The severity of the message, i.e. error, warning etc.
+	 **/
+	static void AddMessageToMessageLog(FString MessageString, EMessageSeverity::Type MessageType)
+	{
+		FText MsgText = FText::FromString(MessageString);
+
+		TSharedRef<FTokenizedMessage> Message = FTokenizedMessage::Create(MessageType);
+		Message->AddToken(FTextToken::Create(MsgText));
+
+		FMessageLog MessageLog("PackagingResults");
+		MessageLog.AddMessage(Message);
+	}
+
+	/**
+	 * Send Error to the Message Log.
+	 *
+	 * @Param MessageString - The error we wish to send to the Message Log.
+	 * @Param MessageType - The severity of the message, i.e. error, warning etc.
+	 **/
+	static void SyncMessageWithMessageLog(FString MessageString, EMessageSeverity::Type MessageType)
+	{
+		DECLARE_CYCLE_STAT(TEXT("FSimpleDelegateGraphTask.SendPackageErrorToMessageLog"),
+		STAT_FSimpleDelegateGraphTask_SendPackageErrorToMessageLog,
+			STATGROUP_TaskGraphTasks);
+
+		// Remove any new line terminators
+		MessageString.ReplaceInline(TEXT("\r"), TEXT(""));
+		MessageString.ReplaceInline(TEXT("\n"), TEXT(""));
+
+		/**
+		 * Dispatch the error from packaging to the message log.
+		 **/
+		FSimpleDelegateGraphTask::CreateAndDispatchWhenReady(
+			FSimpleDelegateGraphTask::FDelegate::CreateStatic(&FPackagingErrorHandler::AddMessageToMessageLog, MessageString, MessageType),
+			GET_STATID(STAT_FSimpleDelegateGraphTask_SendPackageErrorToMessageLog),
+			nullptr, ENamedThreads::GameThread
+			);
+	}
+
+public:
+	/**
+	 * Determine if the output is an error we wish to send to the Message Log.
+	 *
+	 * @Param UATOutput - The current line of output from the UAT package process.
+	 **/
+	static void ProcessAndHandleCookErrorOutput(FString UATOutput)
+	{
+		FString LhsUATOutputMsg, ParsedCookIssue;
+
+		// note: CookResults:Warning: actually outputs some unhandled errors.
+		if (UATOutput.Split(TEXT("CookResults:Warning: "), &LhsUATOutputMsg, &ParsedCookIssue))
+		{
+			SyncMessageWithMessageLog(ParsedCookIssue, EMessageSeverity::Warning);
+		}
+
+		if (UATOutput.Split(TEXT("CookResults:Error: "), &LhsUATOutputMsg, &ParsedCookIssue))
+		{
+			SyncMessageWithMessageLog(ParsedCookIssue, EMessageSeverity::Error);
+		}
+	}
+
+	/**
+	 * Send the UAT Packaging error message to the Message Log.
+	 *
+	 * @Param ErrorCode - The UAT return code we received and wish to display the error message for.
+	 **/
+	static void SendPackagingErrorToMessageLog(int32 ErrorCode)
+	{
+		SyncMessageWithMessageLog(FEditorAnalytics::TranslateErrorCode(ErrorCode), EMessageSeverity::Error);
+	}
+};
+
+
 DECLARE_CYCLE_STAT(TEXT("Requesting FMainFrameActionCallbacks::HandleUatProcessCompleted message dialog to present the error message"), STAT_FMainFrameActionCallbacks_HandleUatProcessCompleted_DialogMessage, STATGROUP_TaskGraphTasks);
 void FMainFrameActionCallbacks::HandleUatProcessCompleted( int32 ReturnCode, TWeakPtr<SNotificationItem> NotificationItemPtr, FText PlatformDisplayName, FText TaskName, EventData Event )
 {
@@ -1191,6 +1275,12 @@ void FMainFrameActionCallbacks::HandleUatProcessCompleted( int32 ReturnCode, TWe
 		ParamArray.Add(FAnalyticsEventAttribute(TEXT("Time"), FPlatformTime::Seconds() - Event.StartTime));
 		FEditorAnalytics::ReportEvent(Event.EventName + TEXT(".Failed"), PlatformDisplayName.ToString(), Event.bProjectHasCode, ReturnCode, ParamArray);
 
+		// Send the error to the Message Log.
+		if (TaskName.EqualTo(LOCTEXT("PackagingTaskName", "Packaging")))
+		{
+			FPackagingErrorHandler::SendPackagingErrorToMessageLog(ReturnCode);
+		}
+
 		// Present a message dialog if we want the error message to be prominent.
 		if (FEditorAnalytics::ShouldElevateMessageThroughDialog(ReturnCode))
 		{
@@ -1213,6 +1303,12 @@ void FMainFrameActionCallbacks::HandleUatProcessOutput( FString Output, TWeakPtr
 	if (!Output.IsEmpty() && !Output.Equals("\r"))
 	{
 		UE_LOG(MainFrameActions, Log, TEXT("%s (%s): %s"), *TaskName.ToString(), *PlatformDisplayName.ToString(), *Output);
+
+		if (TaskName.EqualTo(LOCTEXT("PackagingTaskName", "Packaging")))
+		{
+			// Deal with any cook errors that may have been encountered.
+			FPackagingErrorHandler::ProcessAndHandleCookErrorOutput(Output);
+		}
 	}	
 }
 
