Cancel playing via launcher before starting a package build

diff --git a/Engine/Source/Editor/MainFrame/Private/Frame/MainFrameActions.cpp b/Engine/Source/Editor/MainFrame/Private/Frame/MainFrameActions.cpp
--- a/Engine/Source/Editor/MainFrame/Private/Frame/MainFrameActions.cpp
+++ b/Engine/Source/Editor/MainFrame/Private/Frame/MainFrameActions.cpp
@@ -367,6 +367,8 @@ bool FMainFrameActionCallbacks::PackageBuildConfigurationIsChecked( EProjectPack
 
 void FMainFrameActionCallbacks::PackageProject( const FName InPlatformInfoName )
 {
+	GUnrealEd->CancelPlayingViaLauncher();
+	
 	// does the project have any code?
 	FGameProjectGenerationModule& GameProjectModule = FModuleManager::LoadModuleChecked<FGameProjectGenerationModule>(TEXT("GameProjectGeneration"));
 	bool bProjectHasCode = GameProjectModule.Get().ProjectRequiresBuild(InPlatformInfoName);
diff --git a/Engine/Source/Editor/UnrealEd/Classes/Editor/EditorEngine.h b/Engine/Source/Editor/UnrealEd/Classes/Editor/EditorEngine.h
--- a/Engine/Source/Editor/UnrealEd/Classes/Editor/EditorEngine.h
+++ b/Engine/Source/Editor/UnrealEd/Classes/Editor/EditorEngine.h
@@ -2297,6 +2297,11 @@ public:
 	 */
 	bool IsPlayingViaLauncher() const { return bPlayUsingLauncher && !bIsPlayWorldQueued; }
 
+	/** 
+	 * Cancel playing via the Launcher
+	 */
+	void CancelPlayingViaLauncher();
+
 	/** @return true if the editor is able to launch PIE with online platform support */
 	bool SupportsOnlinePIE() const;
 
diff --git a/Engine/Source/Editor/UnrealEd/Private/PlayLevel.cpp b/Engine/Source/Editor/UnrealEd/Private/PlayLevel.cpp
--- a/Engine/Source/Editor/UnrealEd/Private/PlayLevel.cpp
+++ b/Engine/Source/Editor/UnrealEd/Private/PlayLevel.cpp
@@ -2432,6 +2432,14 @@ bool UEditorEngine::CreatePIEWorldFromLogin(FWorldContext& PieWorldContext, EPla
 	}
 }
 
+void UEditorEngine::CancelPlayingViaLauncher()
+{
+	if (LauncherWorker.IsValid())
+	{
+		LauncherWorker->Cancel();
+	}
+}
+
 bool UEditorEngine::SupportsOnlinePIE() const
 {
 	if (bOnlinePIEEnabled && PIELogins.Num() > 0)
