Button added under File -> Package Project for zipping up the project. Allows you to choose desired location. Filters applied to zip only necessary information.

UE-17086

diff --git a/Engine/Source/Editor/MainFrame/Private/Frame/MainFrameActions.cpp b/Engine/Source/Editor/MainFrame/Private/Frame/MainFrameActions.cpp
--- a/Engine/Source/Editor/MainFrame/Private/Frame/MainFrameActions.cpp
+++ b/Engine/Source/Editor/MainFrame/Private/Frame/MainFrameActions.cpp
@@ -73,6 +73,9 @@ void FMainFrameCommands::RegisterCommands()
 	UI_COMMAND( OpenIDE, "Open IDE", "Opens your C++ code in an integrated development environment.", EUserInterfaceActionType::Button, FInputChord() );
 	ActionList->MapAction( OpenIDE, FExecuteAction::CreateStatic( &FMainFrameActionCallbacks::OpenIDE ), DefaultExecuteAction );
 
+	UI_COMMAND( ZipUpProject, "Zip Up Project", "Zips up the project into a zip file.", EUserInterfaceActionType::Button, FInputChord() );
+	ActionList->MapAction(ZipUpProject, FExecuteAction::CreateStatic( &FMainFrameActionCallbacks::ZipUpProject ), DefaultExecuteAction);
+
 	UI_COMMAND( PackagingSettings, "Packaging Settings...", "Opens the settings for project packaging", EUserInterfaceActionType::Button, FInputChord() );
 	ActionList->MapAction( PackagingSettings, FExecuteAction::CreateStatic( &FMainFrameActionCallbacks::PackagingSettings ), DefaultExecuteAction );
 
@@ -658,6 +661,37 @@ void FMainFrameActionCallbacks::OpenIDE()
 	}
 }
 
+void FMainFrameActionCallbacks::ZipUpProject()
+{
+	bool bOpened = false;
+	TArray<FString> SaveFilenames;
+	IDesktopPlatform* DesktopPlatform = FDesktopPlatformModule::Get();
+	if (DesktopPlatform != NULL)
+	{
+		bOpened = DesktopPlatform->SaveFileDialog(
+			NULL,
+			NSLOCTEXT("UnrealEd", "InterpEd_ExportSoundCueInfoDialogTitle", "Zip file location").ToString(),
+			FPaths::GameDir(),
+			FApp::GetGameName(),
+			TEXT("Zip file|*.zip"),
+			EFileDialogFlags::None,
+			SaveFilenames);
+	}
+
+	if (bOpened)
+	{
+		for (FString FileName : SaveFilenames)
+		{
+			FString ProjectPath = FPaths::IsProjectFilePathSet() ? FPaths::ConvertRelativePathToFull(FPaths::GameDir()) :
+				FPaths::RootDir() / FApp::GetGameName();
+			FString CommandLine = FString::Printf(TEXT("ZipProjectUp -project=\"%s\" -install=\"%s\""), *ProjectPath, *FileName);
+
+			CreateUatTask(CommandLine, LOCTEXT("Windows", "Windows"), LOCTEXT("ZipTaskName", "Zipping Up Project"), 
+				LOCTEXT("ZipTaskShortName", "Zip"), FEditorStyle::GetBrush(TEXT("MainFrame.CookContent")));
+		}
+	}
+}
+
 void FMainFrameActionCallbacks::PackagingSettings()
 {
 	FModuleManager::LoadModuleChecked<ISettingsModule>("Settings").ShowViewer("Project", "Project", "Packaging");
diff --git a/Engine/Source/Editor/MainFrame/Private/Frame/MainFrameActions.h b/Engine/Source/Editor/MainFrame/Private/Frame/MainFrameActions.h
--- a/Engine/Source/Editor/MainFrame/Private/Frame/MainFrameActions.h
+++ b/Engine/Source/Editor/MainFrame/Private/Frame/MainFrameActions.h
@@ -30,6 +30,7 @@ public:
 	TSharedPtr< FUICommandInfo > AddCodeToProject;
 	TSharedPtr< FUICommandInfo > OpenIDE;
 	TSharedPtr< FUICommandInfo > RefreshCodeProject;
+	TSharedPtr< FUICommandInfo > ZipUpProject;
 	TSharedPtr< FUICommandInfo > PackagingSettings;
 	TSharedPtr< FUICommandInfo > LocalizeProject;
 	TArray< TSharedPtr< FUICommandInfo > > SwitchProjectCommands;
@@ -182,6 +183,9 @@ public:
 	/** Opens an IDE to edit c++ code */
 	static void OpenIDE();
 
+	/** Zips up the project */
+	static void ZipUpProject();
+
 	/** Opens the Packaging settings tab */
 	static void PackagingSettings();
 
diff --git a/Engine/Source/Editor/MainFrame/Private/Menus/PackageProjectMenu.h b/Engine/Source/Editor/MainFrame/Private/Menus/PackageProjectMenu.h
--- a/Engine/Source/Editor/MainFrame/Private/Menus/PackageProjectMenu.h
+++ b/Engine/Source/Editor/MainFrame/Private/Menus/PackageProjectMenu.h
@@ -66,6 +66,9 @@ public:
 			}
 		}
 
+		MenuBuilder.AddMenuSeparator();
+		MenuBuilder.AddMenuEntry(FMainFrameCommands::Get().ZipUpProject);
+
 		MenuBuilder.AddMenuSeparator();
 		MenuBuilder.AddSubMenu(
 			LOCTEXT("PackageProjectBuildConfigurationSubMenuLabel", "Build Configuration"),
diff --git a/Engine/Source/Programs/AutomationTool/Scripts/BuildCommonTools.Automation.cs b/Engine/Source/Programs/AutomationTool/Scripts/BuildCommonTools.Automation.cs
--- a/Engine/Source/Programs/AutomationTool/Scripts/BuildCommonTools.Automation.cs
+++ b/Engine/Source/Programs/AutomationTool/Scripts/BuildCommonTools.Automation.cs
@@ -158,3 +158,36 @@ public class BuildCommonTools : BuildCommand
 		return Agenda;
 	}
 }
+
+public class ZipProjectUp : BuildCommand
+{
+    public override void ExecuteBuild()
+    {
+        Log(System.Diagnostics.TraceEventType.Warning, "Started zipping project up");
+
+        // Get the zip path
+        //string ZipPath = ParseParamValue("path", "");
+        //string FullZipFileName = Path.Combine(CmdEnv.LocalRoot, "FullInstall" + StaticGetHostPlatformSuffix(HostPlatform) + ".zip");
+
+        string InstallDirectory = ParseParamValue("install", "");
+        Log("Install directory: {0}", InstallDirectory);
+
+        // Get project directory and format it properly
+        string ProjectDirectory = ParseParamValue("project", "");
+        ProjectDirectory = Path.GetDirectoryName(ProjectDirectory);
+        string path2 = @"[REDACTED_USER_HOME]\Documents\Unreal Projects\MyProject2";
+        Log("Project name: {0}", ProjectDirectory);
+        Log("Project path: {0}", path2);
+
+        // Setup filters
+        FileFilter Filter = new FileFilter();
+        Filter.Include("/Config/...");
+        Filter.Include("/Content/...");
+        Filter.Include("*.uproject");
+
+        ZipFiles(InstallDirectory, ProjectDirectory, Filter);
+        //ZipFiles("d:\\test.zip", path2, Filter);
+
+        Log(System.Diagnostics.TraceEventType.Warning, "Completed zipping project up");
+    }
+}
\ No newline at end of file
