Editor: Added an ensure to the main frame to indicate when it is being incorrectly loaded in a -game build
If this ensure is hit in your game, you need to guard code that is editor-only via either the Editor module type, or checking IsRunningGame()

diff --git a/Engine/Source/Editor/MainFrame/Private/MainFrameModule.cpp b/Engine/Source/Editor/MainFrame/Private/MainFrameModule.cpp
--- a/Engine/Source/Editor/MainFrame/Private/MainFrameModule.cpp
+++ b/Engine/Source/Editor/MainFrame/Private/MainFrameModule.cpp
@@ -608,6 +608,7 @@ void FMainFrameModule::StartupModule( )
 {
 	MRUFavoritesList = NULL;
 
+	ensureMsgf(!IsRunningGame(), TEXT("The MainFrame module should only be loaded when running the editor.  Code that extends the editor, adds menu items, etc... should not run when running in -game mode or in a non-WITH_EDITOR build"));
 	MainFrameHandler = MakeShareable(new FMainFrameHandler);
 
 	FMainFrameCommands::Register();
