Fixed various issues that were only showing up during the process of packaging a project from the editor

I think MainFrame can normally assume that generic commands have been registered by slate application but not in the case of a cook

diff --git a/Engine/Source/Editor/MainFrame/Private/MainFrameModule.cpp b/Engine/Source/Editor/MainFrame/Private/MainFrameModule.cpp
--- a/Engine/Source/Editor/MainFrame/Private/MainFrameModule.cpp
+++ b/Engine/Source/Editor/MainFrame/Private/MainFrameModule.cpp
@@ -6,6 +6,7 @@
 #include "Developer/HotReload/Public/IHotReload.h"
 #include "NotificationManager.h"
 #include "SNotificationList.h"
+#include "GenericCommands.h"
 
 DEFINE_LOG_CATEGORY(LogMainFrame);
 #define LOCTEXT_NAMESPACE "FMainFrameModule"
@@ -611,6 +612,7 @@ void FMainFrameModule::StartupModule( )
 	ensureMsgf(!IsRunningGame(), TEXT("The MainFrame module should only be loaded when running the editor.  Code that extends the editor, adds menu items, etc... should not run when running in -game mode or in a non-WITH_EDITOR build"));
 	MainFrameHandler = MakeShareable(new FMainFrameHandler);
 
+	FGenericCommands::Register();
 	FMainFrameCommands::Register();
 
 	SetLevelNameForWindowTitle(TEXT(""));
diff --git a/Engine/Source/Programs/CrashReporter/CrashReportClient/Private/CrashReportClientUnattended.cpp b/Engine/Source/Programs/CrashReporter/CrashReportClient/Private/CrashReportClientUnattended.cpp
--- a/Engine/Source/Programs/CrashReporter/CrashReportClient/Private/CrashReportClientUnattended.cpp
+++ b/Engine/Source/Programs/CrashReporter/CrashReportClient/Private/CrashReportClientUnattended.cpp
@@ -2,6 +2,7 @@
 
 #include "CrashReportClientApp.h"
 #include "CrashReportClientUnattended.h"
+#include "CrashReportUtil.h"
 
 FCrashReportClientUnattended::FCrashReportClientUnattended(const FPlatformErrorReport& ErrorReport)
 	: Uploader( FCrashReportClientConfig::Get().GetReceiverAddress() )
diff --git a/Engine/Source/Programs/CrashReporter/CrashReportClient/Private/CrashUpload.cpp b/Engine/Source/Programs/CrashReporter/CrashReportClient/Private/CrashUpload.cpp
--- a/Engine/Source/Programs/CrashReporter/CrashReportClient/Private/CrashUpload.cpp
+++ b/Engine/Source/Programs/CrashReporter/CrashReportClient/Private/CrashUpload.cpp
@@ -7,6 +7,7 @@
 #include "PlatformErrorReport.h"
 #include "PendingReports.h"
 #include "XmlFile.h"
+#include "CrashReportUtil.h"
 
 #define LOCTEXT_NAMESPACE "CrashReportClient"
 
