Fixed zipping project path issue on mac. Cleaned up log messages. Display the correct OS for the process.

UE-19001, UE-19003

diff --git a/Engine/Source/Editor/MainFrame/Private/Frame/MainFrameActions.cpp b/Engine/Source/Editor/MainFrame/Private/Frame/MainFrameActions.cpp
--- a/Engine/Source/Editor/MainFrame/Private/Frame/MainFrameActions.cpp
+++ b/Engine/Source/Editor/MainFrame/Private/Frame/MainFrameActions.cpp
@@ -340,7 +340,7 @@ void FMainFrameActionCallbacks::CookContent(const FName InPlatformInfoName)
 		*FUnrealEdMisc::Get().GetExecutableForCommandlets(),
 		*OptionalParams
 	);
-
+	
 	CreateUatTask(CommandLine, PlatformInfo->DisplayName, LOCTEXT("CookingContentTaskName", "Cooking content"), LOCTEXT("CookingTaskName", "Cooking"), FEditorStyle::GetBrush(TEXT("MainFrame.CookContent")));
 }
 
@@ -671,6 +671,16 @@ void FMainFrameActionCallbacks::OpenIDE()
 
 void FMainFrameActionCallbacks::ZipUpProject()
 {
+#if PLATFORM_WINDOWS
+	FText PlatformName = LOCTEXT("Platform Name", "Windows");
+#elif PLATFORM_MAC
+	FText PlatformName = LOCTEXT("Platform Name", "Mac");
+#elif PLATFORM_LINUX
+	FText PlatformName = LOCTEXT("Platform Name", "Linux");
+#else
+	FText PlatformName = LOCTEXT("Platform Name", "Other OS");
+#endif
+
 	bool bOpened = false;
 	TArray<FString> SaveFilenames;
 	IDesktopPlatform* DesktopPlatform = FDesktopPlatformModule::Get();
@@ -690,12 +700,14 @@ void FMainFrameActionCallbacks::ZipUpProject()
 	{
 		for (FString FileName : SaveFilenames)
 		{
-			FString ProjectPath = FPaths::IsProjectFilePathSet() ? FPaths::ConvertRelativePathToFull(FPaths::GameDir()) :
-				FPaths::RootDir() / FApp::GetGameName();
-			FString CommandLine = FString::Printf(TEXT("ZipProjectUp -project=\"%s\" -install=\"%s\""), *ProjectPath, *FileName);
+			// Ensure path is full rather than relative (for macs)
+			FString FinalFileName = FPaths::ConvertRelativePathToFull(FileName);
+			FString ProjectPath = FPaths::IsProjectFilePathSet() ? FPaths::ConvertRelativePathToFull(FPaths::GameDir()) : FPaths::RootDir() / FApp::GetGameName();
+
+			FString CommandLine = FString::Printf(TEXT("ZipProjectUp -project=\"%s\" -install=\"%s\""), *ProjectPath, *FinalFileName);
 
-			CreateUatTask(CommandLine, LOCTEXT("Windows", "Windows"), LOCTEXT("ZipTaskName", "Zipping Up Project"), 
-				LOCTEXT("ZipTaskShortName", "Zip"), FEditorStyle::GetBrush(TEXT("MainFrame.CookContent")));
+			CreateUatTask(CommandLine, PlatformName, LOCTEXT("ZipTaskName", "Zipping Up Project"),
+				LOCTEXT("ZipTaskShortName", "Zip Project Task"), FEditorStyle::GetBrush(TEXT("MainFrame.CookContent")));
 		}
 	}
 }
diff --git a/Engine/Source/Programs/AutomationTool/Scripts/BuildCommonTools.Automation.cs b/Engine/Source/Programs/AutomationTool/Scripts/BuildCommonTools.Automation.cs
--- a/Engine/Source/Programs/AutomationTool/Scripts/BuildCommonTools.Automation.cs
+++ b/Engine/Source/Programs/AutomationTool/Scripts/BuildCommonTools.Automation.cs
@@ -163,21 +163,15 @@ public class ZipProjectUp : BuildCommand
 {
     public override void ExecuteBuild()
     {
-        Log(System.Diagnostics.TraceEventType.Warning, "Started zipping project up");
-
-        // Get the zip path
-        //string ZipPath = ParseParamValue("path", "");
-        //string FullZipFileName = Path.Combine(CmdEnv.LocalRoot, "FullInstall" + StaticGetHostPlatformSuffix(HostPlatform) + ".zip");
-
-        string InstallDirectory = ParseParamValue("install", "");
-        Log("Install directory: {0}", InstallDirectory);
-
-        // Get project directory and format it properly
+        // Get Directories
         string ProjectDirectory = ParseParamValue("project", "");
+        string InstallDirectory = ParseParamValue("install", "");
         ProjectDirectory = Path.GetDirectoryName(ProjectDirectory);
-        string path2 = @"[REDACTED_USER_HOME]\Documents\Unreal Projects\MyProject2";
-        Log("Project name: {0}", ProjectDirectory);
-        Log("Project path: {0}", path2);
+
+        Log("Started zipping project up");
+        Log("Project directory: {0}", ProjectDirectory);
+        Log("Install directory: {0}", InstallDirectory);
+        Log("Packaging up the project...");
 
         // Setup filters
         FileFilter Filter = new FileFilter();
@@ -186,7 +180,6 @@ public class ZipProjectUp : BuildCommand
         Filter.Include("*.uproject");
 
         ZipFiles(InstallDirectory, ProjectDirectory, Filter);
-        //ZipFiles("d:\\test.zip", path2, Filter);
 
         Log(System.Diagnostics.TraceEventType.Warning, "Completed zipping project up");
     }
