Compile game executables for content-only projects if we're not running from a promoted or installed build.

diff --git a/Engine/Source/Developer/LauncherServices/Private/Profiles/LauncherProfile.h b/Engine/Source/Developer/LauncherServices/Private/Profiles/LauncherProfile.h
--- a/Engine/Source/Developer/LauncherServices/Private/Profiles/LauncherProfile.h
+++ b/Engine/Source/Developer/LauncherServices/Private/Profiles/LauncherProfile.h
@@ -875,7 +875,7 @@ public:
 		FInternationalization& I18N = FInternationalization::Get();
 
 		// default build settings
-		BuildGame = false;
+		BuildGame = !FApp::GetEngineIsPromotedBuild() && !FApp::IsEngineInstalled();
 		BuildUAT = !FApp::GetEngineIsPromotedBuild() && !FApp::IsEngineInstalled();
 
 		// default cook settings
diff --git a/Engine/Source/Editor/MainFrame/Private/Frame/MainFrameActions.cpp b/Engine/Source/Editor/MainFrame/Private/Frame/MainFrameActions.cpp
--- a/Engine/Source/Editor/MainFrame/Private/Frame/MainFrameActions.cpp
+++ b/Engine/Source/Editor/MainFrame/Private/Frame/MainFrameActions.cpp
@@ -563,7 +563,7 @@ void FMainFrameActionCallbacks::PackageProject( const FName InPlatformInfoName )
 	}
 
 	// only build if the project has code that might need to be built
-	if (bProjectHasCode && FSourceCodeNavigation::IsCompilerAvailable())
+	if (bProjectHasCode || (!FApp::GetEngineIsPromotedBuild() && !FApp::IsEngineInstalled()))
 	{
 		OptionalParams += TEXT(" -build");
 	}
