Fixed issue with File -> Cook missing some cooking parameters to UAT.

diff --git a/Engine/Source/Editor/MainFrame/Private/Frame/MainFrameActions.cpp b/Engine/Source/Editor/MainFrame/Private/Frame/MainFrameActions.cpp
--- a/Engine/Source/Editor/MainFrame/Private/Frame/MainFrameActions.cpp
+++ b/Engine/Source/Editor/MainFrame/Private/Frame/MainFrameActions.cpp
@@ -293,6 +293,28 @@ const TCHAR* GetUATCompilationFlags()
 		: TEXT("-nocompileeditor");
 }
 
+FString GetCookingOptionalParams()
+{
+	FString OptionalParams;
+	const UProjectPackagingSettings* const PackagingSettings = GetDefault<UProjectPackagingSettings>();
+	if (PackagingSettings->bCookAll)
+	{
+		OptionalParams += TEXT(" -CookAll");
+		// maps only flag only affects cook all
+		if (PackagingSettings->bCookMapsOnly)
+		{
+			OptionalParams += TEXT(" -CookMapsOnly");
+		}
+	}
+
+	if (PackagingSettings->bSkipEditorContent)
+	{
+		OptionalParams += TEXT(" -SKIPEDITORCONTENT");
+	}
+	return OptionalParams;
+}
+
+
 void FMainFrameActionCallbacks::CookContent(const FName InPlatformInfoName)
 {
 	const PlatformInfo::FPlatformInfo* const PlatformInfo = PlatformInfo::FindPlatformInfo(InPlatformInfoName);
@@ -324,6 +346,8 @@ void FMainFrameActionCallbacks::CookContent(const FName InPlatformInfoName)
 		OptionalParams += *PlatformInfo->TargetPlatformName.ToString();
 	}
 
+	OptionalParams += GetCookingOptionalParams();
+
 	const bool bRunningDebug = FParse::Param(FCommandLine::Get(), TEXT("debug"));
 
 	if (bRunningDebug)
@@ -504,20 +528,7 @@ void FMainFrameActionCallbacks::PackageProject( const FName InPlatformInfoName )
 		OptionalParams += TEXT(" -compressed");
 	}
 
-	if ( PackagingSettings->bCookAll )
-	{
-		OptionalParams += TEXT(" -CookAll");
-		// maps only flag only affects cook all
-		if ( PackagingSettings->bCookMapsOnly )
-		{
-			OptionalParams += TEXT(" -CookMapsOnly");
-		}
-	}
-
-	if (PackagingSettings->bSkipEditorContent)
-	{
-		OptionalParams += TEXT(" -SKIPEDITORCONTENT");
-	}
+	OptionalParams += GetCookingOptionalParams();
 
 	if (PackagingSettings->UsePakFile)
 	{
