Made New Project dialog window less high so that it just about fits on a monitor with 1366x768 resolution.

diff --git a/Engine/Source/Editor/MainFrame/Private/MainFrameModule.h b/Engine/Source/Editor/MainFrame/Private/MainFrameModule.h
--- a/Engine/Source/Editor/MainFrame/Private/MainFrameModule.h
+++ b/Engine/Source/Editor/MainFrame/Private/MainFrameModule.h
@@ -125,7 +125,7 @@ protected:
 public:
 
 	/** Get the size of the project browser window */
-	static FVector2D GetProjectBrowserWindowSize() { return FVector2D(1100, 780); }
+	static FVector2D GetProjectBrowserWindowSize() { return FVector2D(1100, 700); }
 
 private:
 
