Copying //UE4/Dev-Editor to //UE4/Dev-Main (Source: //UE4/Dev-Editor @ 2973866)

diff --git a/Engine/Source/Editor/GameProjectGeneration/Private/SProjectBrowser.cpp b/Engine/Source/Editor/GameProjectGeneration/Private/SProjectBrowser.cpp
--- a/Engine/Source/Editor/GameProjectGeneration/Private/SProjectBrowser.cpp
+++ b/Engine/Source/Editor/GameProjectGeneration/Private/SProjectBrowser.cpp
@@ -628,7 +628,7 @@ TSharedPtr<SWidget> SProjectBrowser::OnGetContextMenuContent() const
 
 	FFormatNamedArguments Args;
 	Args.Add(TEXT("FileManagerName"), FPlatformMisc::GetFileManagerName());
-	const FText ExploreToText = FText::Format(NSLOCTEXT("GenericPlatform", "ShowInFileManager", "Show In {FileManagerName}"), Args);
+	const FText ExploreToText = FText::Format(NSLOCTEXT("GenericPlatform", "ShowInFileManager", "Show in {FileManagerName}"), Args);
 
 	MenuBuilder.AddMenuEntry(
 		ExploreToText,
diff --git a/Engine/Source/Editor/GameplayAbilitiesEditor/Private/SGameplayCueEditor.cpp b/Engine/Source/Editor/GameplayAbilitiesEditor/Private/SGameplayCueEditor.cpp
--- a/Engine/Source/Editor/GameplayAbilitiesEditor/Private/SGameplayCueEditor.cpp
+++ b/Engine/Source/Editor/GameplayAbilitiesEditor/Private/SGameplayCueEditor.cpp
@@ -8,7 +8,7 @@
 #include "GameplayCueNotify_Static.h"
 #include "GameplayCueNotify_Actor.h"
 #include "SExpandableArea.h"
-#include "ClassIconFinder.h"
+#include "SlateIconFinder.h"
 #include "EditorClassUtils.h"
 #include "GameplayTagsSettings.h"
 #include "SNotificationList.h"
@@ -255,7 +255,7 @@ void SGameplayCuePickerDialog::OnClassPicked(UClass* InChosenClass)
 /** Generates rows in the list of GameplayCueNotify classes to pick from */
 TSharedRef<ITableRow> SGameplayCuePickerDialog::GenerateListRow(UClass* ItemClass, const TSharedRef<STableViewBase>& OwnerTable)
 {	
-	const FSlateBrush* ItemBrush = FClassIconFinder::FindIconForClass(ItemClass);
+	const FSlateBrush* ItemBrush = FSlateIconFinder::FindIconBrushForClass(ItemClass);
 
 	return 
 	SNew(STableRow< UClass* >, OwnerTable)
diff --git a/Engine/Source/Editor/GraphEditor/Private/GraphEditorActions.cpp b/Engine/Source/Editor/GraphEditor/Private/GraphEditorActions.cpp
--- a/Engine/Source/Editor/GraphEditor/Private/GraphEditorActions.cpp
+++ b/Engine/Source/Editor/GraphEditor/Private/GraphEditorActions.cpp
@@ -3,7 +3,7 @@
 #include "GraphEditorCommon.h"
 #include "GraphEditorActions.h"
 
-#define LOCTEXT_NAMESPACE ""
+#define LOCTEXT_NAMESPACE "GraphEditorCommands"
 
 void FGraphEditorCommandsImpl::RegisterCommands()
 {
diff --git a/Engine/Source/Editor/GraphEditor/Private/KismetNodes/SGraphNodeK2Var.cpp b/Engine/Source/Editor/GraphEditor/Private/KismetNodes/SGraphNodeK2Var.cpp
--- a/Engine/Source/Editor/GraphEditor/Private/KismetNodes/SGraphNodeK2Var.cpp
+++ b/Engine/Source/Editor/GraphEditor/Private/KismetNodes/SGraphNodeK2Var.cpp
@@ -5,7 +5,7 @@
 #include "GraphEditorCommon.h"
 #include "SGraphNodeK2Base.h"
 #include "SGraphNodeK2Var.h"
-#include "ClassIconFinder.h"
+#include "SlateIconFinder.h"
 #include "IDocumentation.h"
 #include "BlueprintEditorUtils.h"
 #include "TutorialMetaData.h"
@@ -114,7 +114,7 @@ void SGraphNodeK2Var::UpdateGraphNode()
 				.AutoWidth()
 			[
 				SNew(SImage)
-					.Image( FClassIconFinder::FindIconForClass( LiteralRef->GetObjectRef() ? LiteralRef->GetObjectRef()->GetClass() : NULL ))
+				.Image(FSlateIconFinder::FindIconBrushForClass(LiteralRef->GetObjectRef() ? LiteralRef->GetObjectRef()->GetClass() : NULL))
 			]
 
 			+SHorizontalBox::Slot()
diff --git a/Engine/Source/Editor/GraphEditor/Private/SGraphNode.cpp b/Engine/Source/Editor/GraphEditor/Private/SGraphNode.cpp
--- a/Engine/Source/Editor/GraphEditor/Private/SGraphNode.cpp
+++ b/Engine/Source/Editor/GraphEditor/Private/SGraphNode.cpp
@@ -712,7 +712,7 @@ void SGraphNode::UpdateGraphNode()
 	const FSlateBrush* IconBrush = NULL;
 	if (GraphNode != NULL && GraphNode->ShowPaletteIconOnNode())
 	{
-		IconBrush = FEditorStyle::GetBrush(GraphNode->GetPaletteIcon(IconColor));
+		IconBrush = GraphNode->GetIconAndTint(IconColor).GetOptionalIcon();
 	}
 
 	TSharedRef<SOverlay> DefaultTitleAreaWidget =
diff --git a/Engine/Source/Editor/IntroTutorials/Private/STutorialsBrowser.cpp b/Engine/Source/Editor/IntroTutorials/Private/STutorialsBrowser.cpp
--- a/Engine/Source/Editor/IntroTutorials/Private/STutorialsBrowser.cpp
+++ b/Engine/Source/Editor/IntroTutorials/Private/STutorialsBrowser.cpp
@@ -1063,10 +1063,10 @@ void STutorialsBrowser::HandleAssetAdded(const FAssetData& InAssetData)
 {
 	if(InAssetData.AssetClass == UBlueprint::StaticClass()->GetFName())
 	{
-		const FString* ParentClassPath = InAssetData.TagsAndValues.Find(TEXT("ParentClass"));
-		if(ParentClassPath != nullptr)
+		const FString ParentClassPath = InAssetData.GetTagValueRef<FString>("ParentClass");
+		if(!ParentClassPath.IsEmpty())
 		{
-			UClass* ParentClass = FindObject<UClass>(NULL, **ParentClassPath);
+			UClass* ParentClass = FindObject<UClass>(NULL, *ParentClassPath);
 			if(ParentClass == UEditorTutorial::StaticClass())
 			{
 				bNeedsRefresh = true;
diff --git a/Engine/Source/Editor/Kismet/Private/BlueprintActionMenuItem.cpp b/Engine/Source/Editor/Kismet/Private/BlueprintActionMenuItem.cpp
--- a/Engine/Source/Editor/Kismet/Private/BlueprintActionMenuItem.cpp
+++ b/Engine/Source/Editor/Kismet/Private/BlueprintActionMenuItem.cpp
@@ -203,7 +203,7 @@ FBlueprintActionMenuItem::FBlueprintActionMenuItem(UBlueprintNodeSpawner const*
 	: FEdGraphSchemaAction(InNodeCategory, UiSpec.MenuName, UiSpec.Tooltip.ToString(), InGrouping, UiSpec.Keywords)
 	, Action(NodeSpawner)
 	, IconTint(UiSpec.IconTint)
-	, IconBrush(FEditorStyle::GetBrush(UiSpec.IconName))
+	, IconBrush(UiSpec.Icon.GetOptionalIcon())
 	, Bindings(InBindings)
 {
 	check(Action != nullptr);
@@ -344,7 +344,7 @@ void FBlueprintActionMenuItem::AppendBindings(const FBlueprintActionContext& Con
 	
 	UpdateSearchData(UiSpec.MenuName, UiSpec.Tooltip.ToString(), FText(), UiSpec.Keywords);
 
-	IconBrush = FEditorStyle::GetBrush(UiSpec.IconName);
+	IconBrush = UiSpec.Icon.GetOptionalIcon();
 	IconTint  = UiSpec.IconTint;
 	DocExcerptRef.DocLink = UiSpec.DocLink;
 	DocExcerptRef.DocExcerptName = UiSpec.DocExcerptTag;
diff --git a/Engine/Source/Editor/Kismet/Private/BlueprintEditor.cpp b/Engine/Source/Editor/Kismet/Private/BlueprintEditor.cpp
--- a/Engine/Source/Editor/Kismet/Private/BlueprintEditor.cpp
+++ b/Engine/Source/Editor/Kismet/Private/BlueprintEditor.cpp
@@ -498,7 +498,7 @@ const FSlateBrush* FBlueprintEditor::GetGlyphForGraph(const UEdGraph* Graph, boo
 FSlateBrush const* FBlueprintEditor::GetVarIconAndColor(const UStruct* VarScope, FName VarName, FSlateColor& IconColorOut)
 {
 	FLinearColor ColorOut;
-	const FSlateBrush* IconBrush = FEditorStyle::GetBrush(UK2Node_Variable::GetVariableIconAndColor(VarScope, VarName, ColorOut));
+	const FSlateBrush* IconBrush = UK2Node_Variable::GetVariableIconAndColor(VarScope, VarName, ColorOut).GetOptionalIcon();
 	IconColorOut = ColorOut;
 	return IconBrush;
 }
@@ -1565,7 +1565,7 @@ void FBlueprintEditor::LoadLibrariesFromAssetRegistry()
 
 			for (int32 AssetIndex = 0; AssetIndex < AssetData.Num(); ++AssetIndex)
 			{
-				FString TagValue = AssetData[ AssetIndex ].TagsAndValues.FindRef(BPTypeName);
+				FString TagValue = AssetData[ AssetIndex ].GetTagValueRef<FString>(BPTypeName);
 
 				//Only check for Blueprint Macros & Functions in the asset data for loading
 				if ( TagValue == BPMacroTypeStr || TagValue == BPFunctionTypeStr )
diff --git a/Engine/Source/Editor/Kismet/Private/BlueprintEditorCommands.cpp b/Engine/Source/Editor/Kismet/Private/BlueprintEditorCommands.cpp
--- a/Engine/Source/Editor/Kismet/Private/BlueprintEditorCommands.cpp
+++ b/Engine/Source/Editor/Kismet/Private/BlueprintEditorCommands.cpp
@@ -10,7 +10,7 @@
 #include "BlueprintFunctionNodeSpawner.h" 
 #include "BlueprintEventNodeSpawner.h" 
 
-#define LOCTEXT_NAMESPACE ""
+#define LOCTEXT_NAMESPACE "BlueprintEditorCommands"
 
 /** UI_COMMAND takes long for the compile to optimize */
 PRAGMA_DISABLE_OPTIMIZATION
diff --git a/Engine/Source/Editor/Kismet/Private/FindInBlueprintManager.cpp b/Engine/Source/Editor/Kismet/Private/FindInBlueprintManager.cpp
--- a/Engine/Source/Editor/Kismet/Private/FindInBlueprintManager.cpp
+++ b/Engine/Source/Editor/Kismet/Private/FindInBlueprintManager.cpp
@@ -52,6 +52,7 @@ const FText FFindInBlueprintSearchTags::FiB_ObjectClass = LOCTEXT("ObjectClass",
 const FText FFindInBlueprintSearchTags::FiB_IsArray = LOCTEXT("IsArray", "IsArray");
 const FText FFindInBlueprintSearchTags::FiB_IsReference = LOCTEXT("IsReference", "IsReference");
 const FText FFindInBlueprintSearchTags::FiB_Glyph = LOCTEXT("Glyph", "Glyph");
+const FText FFindInBlueprintSearchTags::FiB_GlyphStyleSet = LOCTEXT("GlyphStyleSet", "GlyphStyleSet");
 const FText FFindInBlueprintSearchTags::FiB_GlyphColor = LOCTEXT("GlyphColor", "GlyphColor");
 
 const FText FFindInBlueprintSearchTags::FiBMetaDataTag = LOCTEXT("FiBMetaDataTag", "!!FiBMD");
@@ -1490,19 +1491,15 @@ void FFindInBlueprintSearchManager::ExtractUnloadedFiBData(const FAssetData& InA
 	FSearchData NewSearchData;
 
 	NewSearchData.BlueprintPath = InAssetData.ObjectPath;
+	InAssetData.GetTagValue("ParentClass", NewSearchData.ParentClass);
 
-	if (const FString* ParentClass = InAssetData.TagsAndValues.Find(TEXT("ParentClass")))
-	{
-		NewSearchData.ParentClass = *ParentClass;
-	}
-
-	const FString* ImplementedInterfaces = InAssetData.TagsAndValues.Find("ImplementedInterfaces");
-	if(ImplementedInterfaces)
+	const FString ImplementedInterfaces = InAssetData.GetTagValueRef<FString>("ImplementedInterfaces");
+	if(!ImplementedInterfaces.IsEmpty())
 	{
 		FString FullInterface;
 		FString RemainingString;
 		FString InterfaceName;
-		FString CurrentString = *ImplementedInterfaces;
+		FString CurrentString = ImplementedInterfaces;
 		while(CurrentString.Split(TEXT(","), &FullInterface, &RemainingString))
 		{
 			if(FullInterface.Split(TEXT("."), &CurrentString, &InterfaceName, ESearchCase::CaseSensitive, ESearchDir::FromEnd))
diff --git a/Engine/Source/Editor/Kismet/Private/FindInBlueprints.cpp b/Engine/Source/Editor/Kismet/Private/FindInBlueprints.cpp
--- a/Engine/Source/Editor/Kismet/Private/FindInBlueprints.cpp
+++ b/Engine/Source/Editor/Kismet/Private/FindInBlueprints.cpp
@@ -200,7 +200,7 @@ FText FFindInBlueprintsResult::GetDisplayString() const
 
 FFindInBlueprintsGraphNode::FFindInBlueprintsGraphNode(const FText& InValue, TSharedPtr<FFindInBlueprintsResult> InParent)
 	: FFindInBlueprintsResult(InValue, InParent)
-	, GlyphBrush(nullptr)
+	, Glyph("EditorStyle", "")
 	, Class(nullptr)
 {
 }
@@ -224,7 +224,7 @@ FReply FFindInBlueprintsGraphNode::OnClick()
 TSharedRef<SWidget> FFindInBlueprintsGraphNode::CreateIcon() const
 {
 	return 	SNew(SImage)
-		.Image(GlyphBrush)
+		.Image(Glyph.GetOptionalIcon())
 		.ColorAndOpacity(GlyphColor)
 		.ToolTipText( GetCategory() );
 }
@@ -251,7 +251,11 @@ void FFindInBlueprintsGraphNode::ParseSearchInfo(FText InKey, FText InValue)
 	}
 	else if(InKey.CompareTo(FFindInBlueprintSearchTags::FiB_Glyph) == 0)
 	{
-		GlyphBrush = FEditorStyle::GetBrush(*InValue.ToString());
+		Glyph = FSlateIcon(Glyph.GetStyleSetName(), *InValue.ToString());
+	}
+	else if(InKey.CompareTo(FFindInBlueprintSearchTags::FiB_GlyphStyleSet) == 0)
+	{
+		Glyph = FSlateIcon(*InValue.ToString(), Glyph.GetStyleName());
 	}
 	else if(InKey.CompareTo(FFindInBlueprintSearchTags::FiB_GlyphColor) == 0)
 	{
@@ -391,7 +395,7 @@ FFindInBlueprintsProperty::FFindInBlueprintsProperty(const FText& InValue, TShar
 TSharedRef<SWidget> FFindInBlueprintsProperty::CreateIcon() const
 {
 	FLinearColor IconColor = FLinearColor::White;
-	const FSlateBrush* Brush = FEditorStyle::GetBrush(UK2Node_Variable::GetVarIconFromPinType(PinType, IconColor));
+	const FSlateBrush* Brush = UK2Node_Variable::GetVarIconFromPinType(PinType, IconColor).GetOptionalIcon();
 	IconColor = UEdGraphSchema_K2::StaticClass()->GetDefaultObject<UEdGraphSchema_K2>()->GetPinTypeColor(PinType);
 
 	return 	SNew(SImage)
diff --git a/Engine/Source/Editor/Kismet/Private/ImaginaryBlueprintData.cpp b/Engine/Source/Editor/Kismet/Private/ImaginaryBlueprintData.cpp
--- a/Engine/Source/Editor/Kismet/Private/ImaginaryBlueprintData.cpp
+++ b/Engine/Source/Editor/Kismet/Private/ImaginaryBlueprintData.cpp
@@ -622,6 +622,7 @@ ESearchableValueStatus FImaginaryGraphNode::GetSearchabilityStatus(FString InKey
 		return ESearchableValueStatus::CoreDisplayItem;
 	}
 	else if (FindInBlueprintsHelpers::IsTextEqualToString(FFindInBlueprintSearchTags::FiB_Glyph, InKey)
+		|| FindInBlueprintsHelpers::IsTextEqualToString(FFindInBlueprintSearchTags::FiB_GlyphStyleSet, InKey)
 		|| FindInBlueprintsHelpers::IsTextEqualToString(FFindInBlueprintSearchTags::FiB_GlyphColor, InKey)
 		|| FindInBlueprintsHelpers::IsTextEqualToString(FFindInBlueprintSearchTags::FiB_NodeGuid, InKey)
 		)
diff --git a/Engine/Source/Editor/Kismet/Private/SBlueprintPalette.cpp b/Engine/Source/Editor/Kismet/Private/SBlueprintPalette.cpp
--- a/Engine/Source/Editor/Kismet/Private/SBlueprintPalette.cpp
+++ b/Engine/Source/Editor/Kismet/Private/SBlueprintPalette.cpp
@@ -321,7 +321,7 @@ static void GetPaletteItemIcon(TSharedPtr<FEdGraphSchemaAction> ActionIn, UBluep
 
 		// Ask node for a palette icon
 		FLinearColor IconLinearColor = FLinearColor::White;
-		BrushOut = FEditorStyle::GetBrush(NodeTemplate->GetPaletteIcon(IconLinearColor));
+		BrushOut = NodeTemplate->GetIconAndTint(IconLinearColor).GetOptionalIcon();
 		ColorOut = IconLinearColor;
 	}
 	// for MyBlueprint tab specific actions:
diff --git a/Engine/Source/Editor/Kismet/Private/SSCSEditor.cpp b/Engine/Source/Editor/Kismet/Private/SSCSEditor.cpp
--- a/Engine/Source/Editor/Kismet/Private/SSCSEditor.cpp
+++ b/Engine/Source/Editor/Kismet/Private/SSCSEditor.cpp
@@ -1317,7 +1317,7 @@ TSharedRef<SWidget> SSCS_RowWidget::GenerateWidgetForColumn( const FName& Column
 		const FSlateBrush* ComponentIcon = FEditorStyle::GetBrush("SCS.NativeComponent");
 		if(NodePtr->GetComponentTemplate() != NULL)
 		{
-			ComponentIcon = FClassIconFinder::FindIconForClass( NodePtr->GetComponentTemplate()->GetClass(), TEXT("SCS.Component") );
+			ComponentIcon = FSlateIconFinder::FindIconBrushForClass( NodePtr->GetComponentTemplate()->GetClass(), TEXT("SCS.Component") );
 		}
 
 		InlineWidget =
diff --git a/Engine/Source/Editor/Kismet/Private/UserDefinedStructureEditor.cpp b/Engine/Source/Editor/Kismet/Private/UserDefinedStructureEditor.cpp
--- a/Engine/Source/Editor/Kismet/Private/UserDefinedStructureEditor.cpp
+++ b/Engine/Source/Editor/Kismet/Private/UserDefinedStructureEditor.cpp
@@ -93,6 +93,7 @@ public:
 		ViewArgs.bAllowSearch = false;
 		ViewArgs.bHideSelectionTip = false;
 		ViewArgs.bShowActorLabel = false;
+		ViewArgs.NameAreaSettings = FDetailsViewArgs::HideNameArea;
 		ViewArgs.NotifyHook = this;
 
 		DetailsView = PropertyModule.CreateDetailView(ViewArgs);
@@ -359,6 +360,7 @@ TSharedRef<SDockTab> FUserDefinedStructureEditor::SpawnStructureTab(const FSpawn
 		PropertyView->RegisterInstancedCustomPropertyLayout(UUserDefinedStruct::StaticClass(), LayoutStructDetails);
 		PropertyView->SetObject(EditedStruct);
 		Box->AddSlot()
+		.VAlign(EVerticalAlignment::VAlign_Top)
 		[
 			PropertyView.ToSharedRef()
 		];
@@ -376,6 +378,7 @@ TSharedRef<SDockTab> FUserDefinedStructureEditor::SpawnStructureTab(const FSpawn
 		{
 			Box->AddSlot()
 			.VAlign(EVerticalAlignment::VAlign_Top)
+			.Padding(2.0f, 0.0f, 0.0f, 0.0f)
 			[
 				DefaultValueWidget.ToSharedRef()
 			];
@@ -585,9 +588,10 @@ public:
 		if(StructureDetailsSP.IsValid())
 		{
 			FStructureEditorUtils::ChangeVariableType(StructureDetailsSP->GetUserDefinedStruct(), FieldGuid, PinType);
-			if (StructureLayout.IsValid())
+			auto StructureLayoutPin = StructureLayout.Pin();
+			if (StructureLayoutPin.IsValid())
 			{
-				StructureLayout.Pin()->OnPinTypeSelected(PinType);
+				StructureLayoutPin->OnPinTypeSelected(PinType);
 			}
 		}
 	}
diff --git a/Engine/Source/Editor/Kismet/Public/FindInBlueprintManager.h b/Engine/Source/Editor/Kismet/Public/FindInBlueprintManager.h
--- a/Engine/Source/Editor/Kismet/Public/FindInBlueprintManager.h
+++ b/Engine/Source/Editor/Kismet/Public/FindInBlueprintManager.h
@@ -69,6 +69,8 @@ struct KISMET_API FFindInBlueprintSearchTags
 	static const FText FiB_IsReference;
 	/** Glyph icon tag */
 	static const FText FiB_Glyph;
+	/** Style set the glyph belongs to */
+	static const FText FiB_GlyphStyleSet;
 	/** Glyph icon color tag */
 	static const FText FiB_GlyphColor;
 
diff --git a/Engine/Source/Editor/Kismet/Public/FindInBlueprints.h b/Engine/Source/Editor/Kismet/Public/FindInBlueprints.h
--- a/Engine/Source/Editor/Kismet/Public/FindInBlueprints.h
+++ b/Engine/Source/Editor/Kismet/Public/FindInBlueprints.h
@@ -157,7 +157,7 @@ private:
 	FGuid NodeGuid;
 
 	/** The glyph brush for this node */
-	const struct FSlateBrush* GlyphBrush;
+	FSlateIcon Glyph;
 
 	/** The glyph color for this node */
 	FLinearColor GlyphColor;
diff --git a/Engine/Source/Editor/LandscapeEditor/Private/LandscapeEditorCommands.cpp b/Engine/Source/Editor/LandscapeEditor/Private/LandscapeEditorCommands.cpp
--- a/Engine/Source/Editor/LandscapeEditor/Private/LandscapeEditorCommands.cpp
+++ b/Engine/Source/Editor/LandscapeEditor/Private/LandscapeEditorCommands.cpp
@@ -4,7 +4,7 @@
 #include "LandscapeEditorPrivatePCH.h"
 #include "LandscapeEditorCommands.h"
 
-#define LOCTEXT_NAMESPACE ""
+#define LOCTEXT_NAMESPACE "LandscapeEditorCommands"
 
 void FLandscapeEditorCommands::RegisterCommands()
 {
diff --git a/Engine/Source/Editor/LandscapeEditor/Private/LandscapeEditorDetailCustomization_NewLandscape.cpp b/Engine/Source/Editor/LandscapeEditor/Private/LandscapeEditorDetailCustomization_NewLandscape.cpp
--- a/Engine/Source/Editor/LandscapeEditor/Private/LandscapeEditorDetailCustomization_NewLandscape.cpp
+++ b/Engine/Source/Editor/LandscapeEditor/Private/LandscapeEditorDetailCustomization_NewLandscape.cpp
@@ -1526,10 +1526,10 @@ FReply FLandscapeEditorStructCustomization_FLandscapeImportLayer::OnLayerFilenam
 
 bool FLandscapeEditorStructCustomization_FLandscapeImportLayer::ShouldFilterLayerInfo(const FAssetData& AssetData, FName LayerName)
 {
-	const FString* const LayerNameMetaData = AssetData.TagsAndValues.Find("LayerName");
-	if (LayerNameMetaData && !LayerNameMetaData->IsEmpty())
+	const FName LayerNameMetaData = AssetData.GetTagValueRef<FName>("LayerName");
+	if (!LayerNameMetaData.IsNone())
 	{
-		return FName(**LayerNameMetaData) != LayerName;
+		return LayerNameMetaData != LayerName;
 	}
 
 	ULandscapeLayerInfoObject* LayerInfo = CastChecked<ULandscapeLayerInfoObject>(AssetData.GetAsset());
diff --git a/Engine/Source/Editor/LandscapeEditor/Private/LandscapeEditorDetailCustomization_TargetLayers.cpp b/Engine/Source/Editor/LandscapeEditor/Private/LandscapeEditorDetailCustomization_TargetLayers.cpp
--- a/Engine/Source/Editor/LandscapeEditor/Private/LandscapeEditorDetailCustomization_TargetLayers.cpp
+++ b/Engine/Source/Editor/LandscapeEditor/Private/LandscapeEditorDetailCustomization_TargetLayers.cpp
@@ -639,10 +639,10 @@ void FLandscapeEditorCustomNodeBuilder_TargetLayers::OnReimportLayer(const TShar
 
 bool FLandscapeEditorCustomNodeBuilder_TargetLayers::ShouldFilterLayerInfo(const FAssetData& AssetData, FName LayerName)
 {
-	const FString* const LayerNameMetaData = AssetData.TagsAndValues.Find("LayerName");
-	if (LayerNameMetaData && !LayerNameMetaData->IsEmpty())
+	const FName LayerNameMetaData = AssetData.GetTagValueRef<FName>("LayerName");
+	if (!LayerNameMetaData.IsNone())
 	{
-		return FName(**LayerNameMetaData) != LayerName;
+		return LayerNameMetaData != LayerName;
 	}
 
 	ULandscapeLayerInfoObject* LayerInfo = CastChecked<ULandscapeLayerInfoObject>(AssetData.GetAsset());
diff --git a/Engine/Source/Editor/Layers/Private/LayerCollectionViewCommands.h b/Engine/Source/Editor/Layers/Private/LayerCollectionViewCommands.h
--- a/Engine/Source/Editor/Layers/Private/LayerCollectionViewCommands.h
+++ b/Engine/Source/Editor/Layers/Private/LayerCollectionViewCommands.h
@@ -2,7 +2,7 @@
 
 #pragma once
 
-#define LOCTEXT_NAMESPACE ""
+#define LOCTEXT_NAMESPACE "LayersViewCommands"
 
 /**
  * The set of commands supported by the LayersView
diff --git a/Engine/Source/Editor/Layers/Private/SLayerStats.h b/Engine/Source/Editor/Layers/Private/SLayerStats.h
--- a/Engine/Source/Editor/Layers/Private/SLayerStats.h
+++ b/Engine/Source/Editor/Layers/Private/SLayerStats.h
@@ -69,7 +69,7 @@ public:
 					.AutoWidth()
 					[
 						SNew(SImage)
-						.Image( FClassIconFinder::FindIconForClass( ActorStats[ StatsIndex ].Type ) )
+						.Image(FSlateIconFinder::FindIconBrushForClass(ActorStats[StatsIndex].Type))
 						.ColorAndOpacity( this, &SLayerStats::GetForegroundColorForButton, StatsIndex )
 					]
 
diff --git a/Engine/Source/Editor/LevelEditor/Private/LevelEditorContextMenu.cpp b/Engine/Source/Editor/LevelEditor/Private/LevelEditorContextMenu.cpp
--- a/Engine/Source/Editor/LevelEditor/Private/LevelEditorContextMenu.cpp
+++ b/Engine/Source/Editor/LevelEditor/Private/LevelEditorContextMenu.cpp
@@ -18,7 +18,7 @@
 #include "AssetData.h"
 #include "DebuggerCommands.h"
 #include "AssetThumbnail.h"
-#include "ClassIconFinder.h"
+#include "SlateIconFinder.h"
 #include "IPlacementModeModule.h"
 #include "AssetRegistryModule.h"
 #include "EngineUtils.h"
@@ -166,7 +166,7 @@ void FLevelEditorContextMenu::FillMenu( FMenuBuilder& MenuBuilder, TWeakPtr<SLev
 				NAME_None,
 				FText::Format(LOCTEXT("SelectComponentOwner", "Select Owner [{0}]"), FText::FromString(OwnerActor->GetHumanReadableName())),
 				TAttribute<FText>(),
-				FSlateIcon(FEditorStyle::GetStyleSetName(), FClassIconFinder::FindIconNameForClass(OwnerActor->GetClass()))
+				FSlateIconFinder::FindIconForClass(OwnerActor->GetClass())
 				);
 
 			MenuBuilder.AddMenuEntry(FEditorViewportCommands::Get().FocusViewportToSelection);
@@ -236,7 +236,7 @@ void FLevelEditorContextMenu::FillMenu( FMenuBuilder& MenuBuilder, TWeakPtr<SLev
 						NAME_None,
 						FText::Format(LOCTEXT("EditAssociatedAsset", "Edit {0}"), FText::FromString(Asset->GetName())),
 						TAttribute<FText>(),
-						FSlateIcon(FEditorStyle::GetStyleSetName(), FClassIconFinder::FindIconNameForClass(Asset->GetClass()))
+						FSlateIconFinder::FindIconForClass(Asset->GetClass())
 						);
 				}
 				else if (ReferencedAssets.Num() > 1)
diff --git a/Engine/Source/Editor/LevelEditor/Private/LevelEditorCreateActorMenu.cpp b/Engine/Source/Editor/LevelEditor/Private/LevelEditorCreateActorMenu.cpp
--- a/Engine/Source/Editor/LevelEditor/Private/LevelEditorCreateActorMenu.cpp
+++ b/Engine/Source/Editor/LevelEditor/Private/LevelEditorCreateActorMenu.cpp
@@ -96,7 +96,7 @@ class SAssetMenuEntry : public SCompoundWidget
 					ActorTypeDisplayName = MenuItem.FactoryToUse->GetDisplayName();
 				}
 
-				IconBrush = FClassIconFinder::FindIconForClass( MenuItem.FactoryToUse->GetClass() );
+				IconBrush = FSlateIconFinder::FindIconBrushForClass(MenuItem.FactoryToUse->GetClass());
 			}
 
 			if ( DefaultActor != NULL && ( MenuItem.FactoryToUse != NULL || !IsClass ) )
@@ -231,10 +231,10 @@ static void FillAssetAddReplaceActorMenu( FMenuBuilder& MenuBuilder, const FAsse
 		FText Label = MenuItem.FactoryToUse->DisplayName;
 		FText ToolTip = MenuItem.FactoryToUse->DisplayName;
 
-		FName Icon = *FString::Printf( TEXT("ClassIcon.%s"), *MenuItem.FactoryToUse->GetClass()->GetName() );
-		if ( !FEditorStyle::GetOptionalBrush(Icon, nullptr, nullptr) )
+		FSlateIcon Icon = FSlateIconFinder::FindIcon(*FString::Printf(TEXT("ClassIcon.%s"), *MenuItem.FactoryToUse->GetClass()->GetName()));
+		if ( !Icon.IsSet() )
 		{
-			Icon = FClassIconFinder::FindIconNameForActor( DefaultActor ) ;
+			Icon = FClassIconFinder::FindSlateIconForActor(DefaultActor);
 		}
 
 		FUIAction Action;
@@ -247,7 +247,7 @@ static void FillAssetAddReplaceActorMenu( FMenuBuilder& MenuBuilder, const FAsse
 			Action = FUIAction( FExecuteAction::CreateStatic( &FLevelEditorActionCallbacks::AddActor_Clicked, MenuItem.FactoryToUse,  MenuItem.AssetData, CreateMode == EActorCreateMode::Placement ) );
 		}
 
-		MenuBuilder.AddMenuEntry( Label, ToolTip, FSlateIcon(FEditorStyle::GetStyleSetName(), Icon), Action );
+		MenuBuilder.AddMenuEntry( Label, ToolTip, Icon, Action );
 	}
 }
 
diff --git a/Engine/Source/Editor/LevelEditor/Private/SLevelViewportToolBar.cpp b/Engine/Source/Editor/LevelEditor/Private/SLevelViewportToolBar.cpp
--- a/Engine/Source/Editor/LevelEditor/Private/SLevelViewportToolBar.cpp
+++ b/Engine/Source/Editor/LevelEditor/Private/SLevelViewportToolBar.cpp
@@ -918,7 +918,7 @@ TSharedRef<SWidget> SLevelViewportToolBar::GenerateShowMenu() const
 		if( ShowMenu[SFG_Normal].Num() > 0 )
 		{
 			// Generate entries for the standard show flags
-			ShowMenuBuilder.BeginSection("LevelViewportShowFlagsCommon", LOCTEXT("CommonShowFlagHeader", "Common Show flags") );
+			ShowMenuBuilder.BeginSection("LevelViewportShowFlagsCommon", LOCTEXT("CommonShowFlagHeader", "Common Show Flags") );
 			{
 				for( int32 EntryIndex = 0; EntryIndex < ShowMenu[SFG_Normal].Num(); ++EntryIndex )
 				{
@@ -929,7 +929,7 @@ TSharedRef<SWidget> SLevelViewportToolBar::GenerateShowMenu() const
 		}
 
 		// Generate entries for the different show flags groups
-		ShowMenuBuilder.BeginSection("LevelViewportShowFlags", LOCTEXT("AllShowFlagHeader", "All Show flags"));
+		ShowMenuBuilder.BeginSection("LevelViewportShowFlags", LOCTEXT("AllShowFlagHeader", "All Show Flags"));
 		{
 			ShowMenuBuilder.AddSubMenu( LOCTEXT("PostProcessShowFlagsMenu", "Post Processing"), LOCTEXT("PostProcessShowFlagsMenu_ToolTip", "Post process show flags"),
 				FNewMenuDelegate::CreateStatic( &FillShowMenu, ShowMenu[SFG_PostProcess], 0 ) );
diff --git a/Engine/Source/Editor/MainFrame/MainFrame.Build.cs b/Engine/Source/Editor/MainFrame/MainFrame.Build.cs
--- a/Engine/Source/Editor/MainFrame/MainFrame.Build.cs
+++ b/Engine/Source/Editor/MainFrame/MainFrame.Build.cs
@@ -35,6 +35,7 @@ public class MainFrame : ModuleRules
 				"WorkspaceMenuStructure",
 				"MessageLog",
 //				"SearchUI",
+				"UATHelper",
 				"TranslationEditor",
 				"Projects",
                 "DeviceProfileEditor",
diff --git a/Engine/Source/Editor/MainFrame/Private/Frame/MainFrameActions.cpp b/Engine/Source/Editor/MainFrame/Private/Frame/MainFrameActions.cpp
--- a/Engine/Source/Editor/MainFrame/Private/Frame/MainFrameActions.cpp
+++ b/Engine/Source/Editor/MainFrame/Private/Frame/MainFrameActions.cpp
@@ -10,6 +10,7 @@
 #include "EngineBuildSettings.h"
 #include "SourceCodeNavigation.h"
 #include "SOutputLogDialog.h"
+#include "IUATHelperModule.h"
 
 #include "Settings/EditorSettings.h"
 
@@ -373,7 +374,7 @@ void FMainFrameActionCallbacks::CookContent(const FName InPlatformInfoName)
 		*OptionalParams
 	);
 	
-	CreateUatTask(CommandLine, PlatformInfo->DisplayName, LOCTEXT("CookingContentTaskName", "Cooking content"), LOCTEXT("CookingTaskName", "Cooking"), FEditorStyle::GetBrush(TEXT("MainFrame.CookContent")));
+	IUATHelperModule::Get().CreateUatTask(CommandLine, PlatformInfo->DisplayName, LOCTEXT("CookingContentTaskName", "Cooking content"), LOCTEXT("CookingTaskName", "Cooking"), FEditorStyle::GetBrush(TEXT("MainFrame.CookContent")));
 }
 
 bool FMainFrameActionCallbacks::CookContentCanExecute( const FName PlatformInfoName )
@@ -649,7 +650,7 @@ void FMainFrameActionCallbacks::PackageProject( const FName InPlatformInfoName )
 		*OptionalParams
 	);
 
-	CreateUatTask(CommandLine, PlatformInfo->DisplayName, LOCTEXT("PackagingProjectTaskName", "Packaging project"), LOCTEXT("PackagingTaskName", "Packaging"), FEditorStyle::GetBrush(TEXT("MainFrame.PackageProject")));
+	IUATHelperModule::Get().CreateUatTask( CommandLine, PlatformInfo->DisplayName, LOCTEXT("PackagingProjectTaskName", "Packaging project"), LOCTEXT("PackagingTaskName", "Packaging"), FEditorStyle::GetBrush(TEXT("MainFrame.PackageProject")) );
 }
 
 bool FMainFrameActionCallbacks::PackageProjectCanExecute( const FName PlatformInfoName )
@@ -742,7 +743,7 @@ void FMainFrameActionCallbacks::ZipUpProject()
 
 			FString CommandLine = FString::Printf(TEXT("ZipProjectUp -project=\"%s\" -install=\"%s\""), *ProjectPath, *FinalFileName);
 
-			CreateUatTask(CommandLine, PlatformName, LOCTEXT("ZipTaskName", "Zipping Up Project"),
+			IUATHelperModule::Get().CreateUatTask( CommandLine, PlatformName, LOCTEXT("ZipTaskName", "Zipping Up Project"),
 				LOCTEXT("ZipTaskShortName", "Zip Project Task"), FEditorStyle::GetBrush(TEXT("MainFrame.CookContent")));
 		}
 	}
@@ -1067,353 +1068,4 @@ void FMainFrameActionCallbacks::AddMessageLog( const FText& Text, const FText& D
 	MessageLog.Open();
 }
 
-
-void FMainFrameActionCallbacks::CreateUatTask( const FString& CommandLine, const FText& PlatformDisplayName, const FText& TaskName, const FText &TaskShortName, const FSlateBrush* TaskIcon )
-{
-	// make sure that the UAT batch file is in place
-#if PLATFORM_WINDOWS
-	FString RunUATScriptName = TEXT("RunUAT.bat");
-	FString CmdExe = TEXT("cmd.exe");
-#elif PLATFORM_LINUX
-	FString RunUATScriptName = TEXT("RunUAT.sh");
-	FString CmdExe = TEXT("/bin/bash");
-#else
-	FString RunUATScriptName = TEXT("RunUAT.command");
-	FString CmdExe = TEXT("/bin/sh");
-#endif
-
-	FString UatPath = FPaths::ConvertRelativePathToFull(FPaths::EngineDir() / TEXT("Build/BatchFiles") / RunUATScriptName);
-	FGameProjectGenerationModule& GameProjectModule = FModuleManager::LoadModuleChecked<FGameProjectGenerationModule>(TEXT("GameProjectGeneration"));
-	bool bHasCode = GameProjectModule.Get().ProjectHasCodeFiles();
-
-	if (!FPaths::FileExists(UatPath))
-	{
-		FFormatNamedArguments Arguments;
-		Arguments.Add(TEXT("File"), FText::FromString(UatPath));
-		FMessageDialog::Open(EAppMsgType::Ok, FText::Format(LOCTEXT("RequiredFileNotFoundMessage", "A required file could not be found:\n{File}"), Arguments));
-
-		TArray<FAnalyticsEventAttribute> ParamArray;
-		ParamArray.Add(FAnalyticsEventAttribute(TEXT("Time"), 0.0));
-		FString EventName = (CommandLine.Contains(TEXT("-package")) ? TEXT("Editor.Package") : TEXT("Editor.Cook"));
-		FEditorAnalytics::ReportEvent(EventName + TEXT(".Failed"), PlatformDisplayName.ToString(), bHasCode, EAnalyticsErrorCodes::UATNotFound, ParamArray);
-		
-		return;
-	}
-
-#if PLATFORM_WINDOWS
-	FString FullCommandLine = FString::Printf(TEXT("/c \"\"%s\" %s\""), *UatPath, *CommandLine);
-#else
-	FString FullCommandLine = FString::Printf(TEXT("\"%s\" %s"), *UatPath, *CommandLine);
-#endif
-
-	TSharedPtr<FMonitoredProcess> UatProcess = MakeShareable(new FMonitoredProcess(CmdExe, FullCommandLine, true));
-
-	// create notification item
-	FFormatNamedArguments Arguments;
-	Arguments.Add(TEXT("Platform"), PlatformDisplayName);
-	Arguments.Add(TEXT("TaskName"), TaskName);
-	FNotificationInfo Info( FText::Format( LOCTEXT("UatTaskInProgressNotification", "{TaskName} for {Platform}..."), Arguments) );
-	
-	Info.Image = TaskIcon;
-	Info.bFireAndForget = false;
-	Info.ExpireDuration = 3.0f;
-	Info.Hyperlink = FSimpleDelegate::CreateStatic(&FMainFrameActionCallbacks::HandleUatHyperlinkNavigate);
-	Info.HyperlinkText = LOCTEXT("ShowOutputLogHyperlink", "Show Output Log");
-	Info.ButtonDetails.Add(
-		FNotificationButtonInfo(
-			LOCTEXT("UatTaskCancel", "Cancel"),
-			LOCTEXT("UatTaskCancelToolTip", "Cancels execution of this task."),
-			FSimpleDelegate::CreateStatic(&FMainFrameActionCallbacks::HandleUatCancelButtonClicked, UatProcess)
-		)
-	);
-
-	TSharedPtr<SNotificationItem> NotificationItem = FSlateNotificationManager::Get().AddNotification(Info);
-
-	if (!NotificationItem.IsValid())
-	{
-		return;
-	}
-
-	FString EventName = (CommandLine.Contains(TEXT("-package")) ? TEXT("Editor.Package") : TEXT("Editor.Cook"));
-	FEditorAnalytics::ReportEvent(EventName + TEXT(".Start"), PlatformDisplayName.ToString(), bHasCode);
-
-	NotificationItem->SetCompletionState(SNotificationItem::CS_Pending);
-
-	// launch the packager
-	TWeakPtr<SNotificationItem> NotificationItemPtr(NotificationItem);
-
-	EventData Data;
-	Data.StartTime = FPlatformTime::Seconds();
-	Data.EventName = EventName;
-	Data.bProjectHasCode = bHasCode;
-	UatProcess->OnCanceled().BindStatic(&FMainFrameActionCallbacks::HandleUatProcessCanceled, NotificationItemPtr, PlatformDisplayName, TaskShortName, Data);
-	UatProcess->OnCompleted().BindStatic(&FMainFrameActionCallbacks::HandleUatProcessCompleted, NotificationItemPtr, PlatformDisplayName, TaskShortName, Data);
-	UatProcess->OnOutput().BindStatic(&FMainFrameActionCallbacks::HandleUatProcessOutput, NotificationItemPtr, PlatformDisplayName, TaskShortName);
-
-	TWeakPtr<FMonitoredProcess> UatProcessPtr(UatProcess);
-	FEditorDelegates::OnShutdownPostPackagesSaved.Add(FSimpleDelegate::CreateStatic(&FMainFrameActionCallbacks::HandleUatCancelButtonClicked, UatProcessPtr));
-
-	if (UatProcess->Launch())
-	{
-		GEditor->PlayEditorSound(TEXT("/Engine/EditorSounds/Notifications/CompileStart_Cue.CompileStart_Cue"));
-	}
-	else
-	{
-		GEditor->PlayEditorSound(TEXT("/Engine/EditorSounds/Notifications/CompileFailed_Cue.CompileFailed_Cue"));
-
-		NotificationItem->SetText(LOCTEXT("UatLaunchFailedNotification", "Failed to launch Unreal Automation Tool (UAT)!"));
-		NotificationItem->SetCompletionState(SNotificationItem::CS_Fail);
-		NotificationItem->ExpireAndFadeout();
-
-		TArray<FAnalyticsEventAttribute> ParamArray;
-		ParamArray.Add(FAnalyticsEventAttribute(TEXT("Time"), 0.0));
-		FEditorAnalytics::ReportEvent(EventName + TEXT(".Failed"), PlatformDisplayName.ToString(), bHasCode, EAnalyticsErrorCodes::UATLaunchFailure, ParamArray);
-	}
-}
-
-/* FMainFrameActionCallbacks callbacks
- *****************************************************************************/
-
-class FMainFrameActionsNotificationTask
-{
-public:
-
-	FMainFrameActionsNotificationTask( TWeakPtr<SNotificationItem> InNotificationItemPtr, SNotificationItem::ECompletionState InCompletionState, const FText& InText )
-		: CompletionState(InCompletionState)
-		, NotificationItemPtr(InNotificationItemPtr)
-		, Text(InText)
-	{ }
-
-	void DoTask( ENamedThreads::Type CurrentThread, const FGraphEventRef& MyCompletionGraphEvent )
-	{
-		if (NotificationItemPtr.IsValid())
-		{
-			if (CompletionState == SNotificationItem::CS_Fail)
-			{
-				GEditor->PlayEditorSound(TEXT("/Engine/EditorSounds/Notifications/CompileFailed_Cue.CompileFailed_Cue"));
-			}
-			else
-			{
-				GEditor->PlayEditorSound(TEXT("/Engine/EditorSounds/Notifications/CompileSuccess_Cue.CompileSuccess_Cue"));
-			}
-			
-			TSharedPtr<SNotificationItem> NotificationItem = NotificationItemPtr.Pin();
-			NotificationItem->SetText(Text);
-			NotificationItem->SetCompletionState(CompletionState);
-			NotificationItem->ExpireAndFadeout();
-		}
-	}
-
-	static ESubsequentsMode::Type GetSubsequentsMode() { return ESubsequentsMode::TrackSubsequents; }
-	ENamedThreads::Type GetDesiredThread( ) { return ENamedThreads::GameThread; }
-	FORCEINLINE TStatId GetStatId() const
-	{
-		RETURN_QUICK_DECLARE_CYCLE_STAT(FMainFrameActionsNotificationTask, STATGROUP_TaskGraphTasks);
-	}
-
-private:
-
-	SNotificationItem::ECompletionState CompletionState;
-	TWeakPtr<SNotificationItem> NotificationItemPtr;
-	FText Text;
-};
-
-
-void FMainFrameActionCallbacks::HandleUatHyperlinkNavigate( )
-{
-	FGlobalTabmanager::Get()->InvokeTab(FName("OutputLog"));
-}
-
-
-void FMainFrameActionCallbacks::HandleUatCancelButtonClicked( TSharedPtr<FMonitoredProcess> PackagerProcess )
-{
-	if (PackagerProcess.IsValid())
-	{
-		PackagerProcess->Cancel(true);
-	}
-}
-
-void FMainFrameActionCallbacks::HandleUatCancelButtonClicked(TWeakPtr<FMonitoredProcess> PackagerProcessPtr)
-{
-	TSharedPtr<FMonitoredProcess> PackagerProcess = PackagerProcessPtr.Pin();
-	if (PackagerProcess.IsValid())
-	{
-		PackagerProcess->Cancel(true);
-	}
-}
-
-void FMainFrameActionCallbacks::HandleUatProcessCanceled( TWeakPtr<SNotificationItem> NotificationItemPtr, FText PlatformDisplayName, FText TaskName, EventData Event )
-{
-	FFormatNamedArguments Arguments;
-	Arguments.Add(TEXT("Platform"), PlatformDisplayName);
-	Arguments.Add(TEXT("TaskName"), TaskName);
-
-	TGraphTask<FMainFrameActionsNotificationTask>::CreateTask().ConstructAndDispatchWhenReady(
-		NotificationItemPtr,
-		SNotificationItem::CS_Fail,
-		FText::Format(LOCTEXT("UatProcessFailedNotification", "{TaskName} canceled!"), Arguments)
-	);
-
-	TArray<FAnalyticsEventAttribute> ParamArray;
-	ParamArray.Add(FAnalyticsEventAttribute(TEXT("Time"), FPlatformTime::Seconds() - Event.StartTime));
-	FEditorAnalytics::ReportEvent(Event.EventName + TEXT(".Canceled"), PlatformDisplayName.ToString(), Event.bProjectHasCode, ParamArray);
-//	FMessageLog("PackagingResults").Warning(FText::Format(LOCTEXT("UatProcessCanceledMessageLog", "{TaskName} for {Platform} canceled by user"), Arguments));
-}
-
-
-/**
- * Helper class to deal with packaging issues encountered in UAT.
- **/
-class FPackagingErrorHandler
-{
-private:
-
-	/**
-	 * Create a message to send to the Message Log.
-	 *
-	 * @Param MessageString - The error we wish to send to the Message Log.
-	 * @Param MessageType - The severity of the message, i.e. error, warning etc.
-	 **/
-	static void AddMessageToMessageLog(FString MessageString, EMessageSeverity::Type MessageType)
-	{
-		FText MsgText = FText::FromString(MessageString);
-
-		TSharedRef<FTokenizedMessage> Message = FTokenizedMessage::Create(MessageType);
-		Message->AddToken(FTextToken::Create(MsgText));
-
-		FMessageLog MessageLog("PackagingResults");
-		MessageLog.AddMessage(Message);
-	}
-
-	/**
-	 * Send Error to the Message Log.
-	 *
-	 * @Param MessageString - The error we wish to send to the Message Log.
-	 * @Param MessageType - The severity of the message, i.e. error, warning etc.
-	 **/
-	static void SyncMessageWithMessageLog(FString MessageString, EMessageSeverity::Type MessageType)
-	{
-		DECLARE_CYCLE_STAT(TEXT("FSimpleDelegateGraphTask.SendPackageErrorToMessageLog"),
-		STAT_FSimpleDelegateGraphTask_SendPackageErrorToMessageLog,
-			STATGROUP_TaskGraphTasks);
-
-		// Remove any new line terminators
-		MessageString.ReplaceInline(TEXT("\r"), TEXT(""));
-		MessageString.ReplaceInline(TEXT("\n"), TEXT(""));
-
-		/**
-		 * Dispatch the error from packaging to the message log.
-		 **/
-		FSimpleDelegateGraphTask::CreateAndDispatchWhenReady(
-			FSimpleDelegateGraphTask::FDelegate::CreateStatic(&FPackagingErrorHandler::AddMessageToMessageLog, MessageString, MessageType),
-			GET_STATID(STAT_FSimpleDelegateGraphTask_SendPackageErrorToMessageLog),
-			nullptr, ENamedThreads::GameThread
-			);
-	}
-
-public:
-	/**
-	 * Determine if the output is an error we wish to send to the Message Log.
-	 *
-	 * @Param UATOutput - The current line of output from the UAT package process.
-	 **/
-	static void ProcessAndHandleCookErrorOutput(FString UATOutput)
-	{
-		FString LhsUATOutputMsg, ParsedCookIssue;
-
-		// note: CookResults:Warning: actually outputs some unhandled errors.
-		if (UATOutput.Split(TEXT("CookResults:Warning: "), &LhsUATOutputMsg, &ParsedCookIssue))
-		{
-			SyncMessageWithMessageLog(ParsedCookIssue, EMessageSeverity::Warning);
-		}
-
-		if (UATOutput.Split(TEXT("CookResults:Error: "), &LhsUATOutputMsg, &ParsedCookIssue))
-		{
-			SyncMessageWithMessageLog(ParsedCookIssue, EMessageSeverity::Error);
-		}
-	}
-
-	/**
-	 * Send the UAT Packaging error message to the Message Log.
-	 *
-	 * @Param ErrorCode - The UAT return code we received and wish to display the error message for.
-	 **/
-	static void SendPackagingErrorToMessageLog(int32 ErrorCode)
-	{
-		SyncMessageWithMessageLog(FEditorAnalytics::TranslateErrorCode(ErrorCode), EMessageSeverity::Error);
-	}
-};
-
-
-DECLARE_CYCLE_STAT(TEXT("Requesting FMainFrameActionCallbacks::HandleUatProcessCompleted message dialog to present the error message"), STAT_FMainFrameActionCallbacks_HandleUatProcessCompleted_DialogMessage, STATGROUP_TaskGraphTasks);
-void FMainFrameActionCallbacks::HandleUatProcessCompleted( int32 ReturnCode, TWeakPtr<SNotificationItem> NotificationItemPtr, FText PlatformDisplayName, FText TaskName, EventData Event )
-{
-	FFormatNamedArguments Arguments;
-	Arguments.Add(TEXT("Platform"), PlatformDisplayName);
-	Arguments.Add(TEXT("TaskName"), TaskName);
-
-	if (ReturnCode == 0)
-	{
-		TGraphTask<FMainFrameActionsNotificationTask>::CreateTask().ConstructAndDispatchWhenReady(
-			NotificationItemPtr,
-			SNotificationItem::CS_Success,
-			FText::Format(LOCTEXT("UatProcessSucceededNotification", "{TaskName} complete!"), Arguments)
-		);
-
-		TArray<FAnalyticsEventAttribute> ParamArray;
-		ParamArray.Add(FAnalyticsEventAttribute(TEXT("Time"), FPlatformTime::Seconds() - Event.StartTime));
-		FEditorAnalytics::ReportEvent(Event.EventName + TEXT(".Completed"), PlatformDisplayName.ToString(), Event.bProjectHasCode, ParamArray);
-
-//		FMessageLog("PackagingResults").Info(FText::Format(LOCTEXT("UatProcessSuccessMessageLog", "{TaskName} for {Platform} completed successfully"), Arguments));
-	}
-	else
-	{
-		TGraphTask<FMainFrameActionsNotificationTask>::CreateTask().ConstructAndDispatchWhenReady(
-			NotificationItemPtr,
-			SNotificationItem::CS_Fail,
-			FText::Format(LOCTEXT("PackagerFailedNotification", "{TaskName} failed!"), Arguments)
-		);
-
-		TArray<FAnalyticsEventAttribute> ParamArray;
-		ParamArray.Add(FAnalyticsEventAttribute(TEXT("Time"), FPlatformTime::Seconds() - Event.StartTime));
-		FEditorAnalytics::ReportEvent(Event.EventName + TEXT(".Failed"), PlatformDisplayName.ToString(), Event.bProjectHasCode, ReturnCode, ParamArray);
-
-		// Send the error to the Message Log.
-		if (TaskName.EqualTo(LOCTEXT("PackagingTaskName", "Packaging")))
-		{
-			FPackagingErrorHandler::SendPackagingErrorToMessageLog(ReturnCode);
-		}
-
-		// Present a message dialog if we want the error message to be prominent.
-		if (FEditorAnalytics::ShouldElevateMessageThroughDialog(ReturnCode))
-		{
-			FSimpleDelegateGraphTask::CreateAndDispatchWhenReady(
-				FSimpleDelegateGraphTask::FDelegate::CreateLambda([=](){
-					FMessageDialog::Open(EAppMsgType::Ok, FText::FromString(FEditorAnalytics::TranslateErrorCode(ReturnCode)));
-				}),
-				GET_STATID(STAT_FMainFrameActionCallbacks_HandleUatProcessCompleted_DialogMessage),
-				nullptr,
-				ENamedThreads::GameThread
-			);
-		}
-//		FMessageLog("PackagingResults").Info(FText::Format(LOCTEXT("UatProcessFailedMessageLog", "{TaskName} for {Platform} failed"), Arguments));
-	}
-}
-
-
-void FMainFrameActionCallbacks::HandleUatProcessOutput( FString Output, TWeakPtr<SNotificationItem> NotificationItemPtr, FText PlatformDisplayName, FText TaskName )
-{
-	if (!Output.IsEmpty() && !Output.Equals("\r"))
-	{
-		UE_LOG(MainFrameActions, Log, TEXT("%s (%s): %s"), *TaskName.ToString(), *PlatformDisplayName.ToString(), *Output);
-
-		if (TaskName.EqualTo(LOCTEXT("PackagingTaskName", "Packaging")))
-		{
-			// Deal with any cook errors that may have been encountered.
-			FPackagingErrorHandler::ProcessAndHandleCookErrorOutput(Output);
-		}
-	}	
-}
-
-
 #undef LOCTEXT_NAMESPACE
diff --git a/Engine/Source/Editor/MainFrame/Private/Frame/MainFrameActions.h b/Engine/Source/Editor/MainFrame/Private/Frame/MainFrameActions.h
--- a/Engine/Source/Editor/MainFrame/Private/Frame/MainFrameActions.h
+++ b/Engine/Source/Editor/MainFrame/Private/Frame/MainFrameActions.h
@@ -254,51 +254,6 @@ protected:
 	 */
 	static void AddMessageLog( const FText& Text, const FText& Detail, const FString& TutorialLink );
 
-	/**
-	 * Creates an asynchronous UAT task.
-	 *
-	 * @param CommandLine - The command line for the UAT task.
-	 * @param PlatformDisplayName - The display name of the platform that the task is running for.
-	 * @param TaskName - The human readable name of the task.
-	 * @param TaskIcon - The icon for the task (to show in the notification item).
-	 */
-	static void CreateUatTask( const FString& CommandLine, const FText& PlatformDisplayName, const FText& TaskName, const FText &TaskShortName, const FSlateBrush* TaskIcon );
-
-	/**
-	 * Translate the given error code to human readable text
-	 *
-	 * @param ErrorCode - the error code returned from UAT task
-	 *
-	 * @return the human readable error string
-	 */
-	static FString TranslateUATError(int32 ErrorCode);
-
-private:
-
-	struct EventData
-	{
-		FString EventName;
-		bool bProjectHasCode;
-		double StartTime;
-	};
-
-	// Handles clicking the packager notification item's Cancel button.
-	static void HandleUatCancelButtonClicked( TSharedPtr<FMonitoredProcess> PackagerProcess );
-
-	static void HandleUatCancelButtonClicked(TWeakPtr<FMonitoredProcess> PackagerProcess);
-
-	// Handles clicking the hyper link on a packager notification item.
-	static void HandleUatHyperlinkNavigate( );
-
-	// Handles canceled packager processes.
-	static void HandleUatProcessCanceled( TWeakPtr<class SNotificationItem> NotificationItemPtr, FText PlatformDisplayName, FText TaskName, EventData Event );
-
-	// Handles the completion of a packager process.
-	static void HandleUatProcessCompleted( int32 ReturnCode, TWeakPtr<class SNotificationItem> NotificationItemPtr, FText PlatformDisplayName, FText TaskName, EventData Event );
-
-	// Handles packager process output.
-	static void HandleUatProcessOutput( FString Output, TWeakPtr<class SNotificationItem> NotificationItemPtr, FText PlatformDisplayName, FText TaskName );
-
 private:
 
 	/** The notification in place while we choose packages to check in */
diff --git a/Engine/Source/Editor/MaterialEditor/Private/MaterialEditorActions.cpp b/Engine/Source/Editor/MaterialEditor/Private/MaterialEditorActions.cpp
--- a/Engine/Source/Editor/MaterialEditor/Private/MaterialEditorActions.cpp
+++ b/Engine/Source/Editor/MaterialEditor/Private/MaterialEditorActions.cpp
@@ -6,7 +6,7 @@
 #include "MaterialEditorActions.h"
 #include "SMaterialEditorViewport.h"
 
-#define LOCTEXT_NAMESPACE ""
+#define LOCTEXT_NAMESPACE "MaterialEditorCommands"
 
 void FMaterialEditorCommands::RegisterCommands()
 {
diff --git a/Engine/Source/Editor/Matinee/Private/MatineeActions.cpp b/Engine/Source/Editor/Matinee/Private/MatineeActions.cpp
--- a/Engine/Source/Editor/Matinee/Private/MatineeActions.cpp
+++ b/Engine/Source/Editor/Matinee/Private/MatineeActions.cpp
@@ -4,7 +4,7 @@
 #include "MatineeActions.h"
 #include "Matinee.h"
 
-#define LOCTEXT_NAMESPACE ""
+#define LOCTEXT_NAMESPACE "MatineeCommands"
 
 /** UI_COMMAND takes long for the compile to optimize */
 PRAGMA_DISABLE_OPTIMIZATION
diff --git a/Engine/Source/Editor/MeshPaint/Private/MeshPaintEdMode.cpp b/Engine/Source/Editor/MeshPaint/Private/MeshPaintEdMode.cpp
--- a/Engine/Source/Editor/MeshPaint/Private/MeshPaintEdMode.cpp
+++ b/Engine/Source/Editor/MeshPaint/Private/MeshPaintEdMode.cpp
@@ -1328,6 +1328,12 @@ static bool PropagateColorsToRawMesh(UStaticMesh* StaticMesh, int32 LODIndex, FS
 	}
 	else
 	{
+		// If there's no raw mesh data, don't try to do any fixup here
+		if (SrcModel.RawMeshBulkData->IsEmpty())
+		{
+			return false;
+		}
+
 		// Fall back to mapping based on position.
 		FRawMesh RawMesh;
 		SrcModel.RawMeshBulkData->LoadRawMesh(RawMesh);
@@ -1555,6 +1561,10 @@ void FEdModeMeshPaint::PaintMeshVertices(
 					// Try using the mapping generated when building the mesh.
 					if(PropagateColorsToRawMesh(StaticMesh, PaintingMeshLODIndex, *InstanceMeshLODInfo))
 					{
+						for (int32 LODIndex = 1; LODIndex < StaticMesh->RenderData->LODResources.Num(); LODIndex++)
+						{
+							PropagateColorsToRawMesh(StaticMesh, LODIndex, *InstanceMeshLODInfo);
+						}
 						RemoveComponentInstanceVertexColors(StaticMeshComponent);
 						StaticMesh->Build();
 					}
@@ -3690,8 +3700,7 @@ bool FEdModeMeshPaint::RequiresInstanceVertexColorsFixup() const
 			for(const auto& StaticMeshComponent : StaticMeshComponents)
 			{
 				// If a static mesh component was found and it requires fixup, exit out and indicate as such
-				TArray<int32> LODsToFixup;
-				if( StaticMeshComponent && StaticMeshComponent->RequiresOverrideVertexColorsFixup( LODsToFixup ) )
+				if( StaticMeshComponent && StaticMeshComponent->RequiresOverrideVertexColorsFixup() )
 				{
 					bRequiresFixup = true;
 					break;
diff --git a/Engine/Source/Editor/MeshPaint/Private/MeshPaintStaticMeshAdapter.cpp b/Engine/Source/Editor/MeshPaint/Private/MeshPaintStaticMeshAdapter.cpp
--- a/Engine/Source/Editor/MeshPaint/Private/MeshPaintStaticMeshAdapter.cpp
+++ b/Engine/Source/Editor/MeshPaint/Private/MeshPaintStaticMeshAdapter.cpp
@@ -188,8 +188,11 @@ void FMeshPaintGeometryAdapterForStaticMeshes::OnRemoved()
 
 		for (const auto& Info : StaticMeshReferencers->Referencers)
 		{
-			Info.StaticMeshComponent->BodyInstance.SetCollisionEnabled(Info.CachedCollisionType, false);
-			Info.StaticMeshComponent->RecreatePhysicsState();
+			if (Info.StaticMeshComponent)
+			{
+				Info.StaticMeshComponent->BodyInstance.SetCollisionEnabled(Info.CachedCollisionType, false);
+				Info.StaticMeshComponent->RecreatePhysicsState();
+			}
 		}
 
 		StaticMeshReferencers->Referencers.Reset();
@@ -266,6 +269,11 @@ void FMeshPaintGeometryAdapterForStaticMeshes::AddReferencedObjects(FReferenceCo
 	FStaticMeshReferencers* StaticMeshReferencers = MeshToComponentMap.Find(ReferencedStaticMesh);
 	check(StaticMeshReferencers);
 	Collector.AddReferencedObject(StaticMeshReferencers->RestoreBodySetup);
+
+	for (auto& Info : StaticMeshReferencers->Referencers)
+	{
+		Collector.AddReferencedObject(Info.StaticMeshComponent);
+	}
 }
 
 
diff --git a/Engine/Source/Editor/MovieSceneTools/Private/TrackEditors/AudioTrackEditor.cpp b/Engine/Source/Editor/MovieSceneTools/Private/TrackEditors/AudioTrackEditor.cpp
--- a/Engine/Source/Editor/MovieSceneTools/Private/TrackEditors/AudioTrackEditor.cpp
+++ b/Engine/Source/Editor/MovieSceneTools/Private/TrackEditors/AudioTrackEditor.cpp
@@ -809,7 +809,7 @@ void FAudioTrackEditor::BuildTrackContextMenu( FMenuBuilder& MenuBuilder, UMovie
 	}
 	UMovieSceneAudioTrack* AudioTrack = Cast<UMovieSceneAudioTrack>( Track );
 	MenuBuilder.AddMenuEntry(
-		NSLOCTEXT( "Sequencer", "PasteMatineeSoundTrack", "Paste Matinee sound track" ),
+		NSLOCTEXT( "Sequencer", "PasteMatineeSoundTrack", "Paste Matinee Sound Track" ),
 		NSLOCTEXT( "Sequencer", "PasteMatineeSoundTrackTooltip", "Pastes keys from a Matinee sound track into this track." ),
 		FSlateIcon(),
 		FUIAction(
diff --git a/Engine/Source/Editor/MovieSceneTools/Private/TrackEditors/PropertyTrackEditors/ColorPropertyTrackEditor.cpp b/Engine/Source/Editor/MovieSceneTools/Private/TrackEditors/PropertyTrackEditors/ColorPropertyTrackEditor.cpp
--- a/Engine/Source/Editor/MovieSceneTools/Private/TrackEditors/PropertyTrackEditors/ColorPropertyTrackEditor.cpp
+++ b/Engine/Source/Editor/MovieSceneTools/Private/TrackEditors/PropertyTrackEditors/ColorPropertyTrackEditor.cpp
@@ -88,8 +88,8 @@ void FColorPropertyTrackEditor::BuildTrackContextMenu( FMenuBuilder& MenuBuilder
 	}
 	UMovieSceneColorTrack* ColorTrack = Cast<UMovieSceneColorTrack>( Track );
 	MenuBuilder.AddMenuEntry(
-		NSLOCTEXT( "Sequencer", "PasteMatineeTrack", "Paste Matinee Color Track" ),
-		NSLOCTEXT( "Sequencer", "PasteMatineeTrackTooltip", "Pastes keys from a Matinee color track into this track." ),
+		NSLOCTEXT( "Sequencer", "PasteMatineeColorTrack", "Paste Matinee Color Track" ),
+		NSLOCTEXT( "Sequencer", "PasteMatineeColorTrackTooltip", "Pastes keys from a Matinee color track into this track." ),
 		FSlateIcon(),
 		FUIAction(
 			FExecuteAction::CreateStatic( &CopyInterpColorTrack, GetSequencer().ToSharedRef(), ColorPropTrack, ColorTrack ),
diff --git a/Engine/Source/Editor/MovieSceneTools/Private/TrackEditors/TransformTrackEditor.cpp b/Engine/Source/Editor/MovieSceneTools/Private/TrackEditors/TransformTrackEditor.cpp
--- a/Engine/Source/Editor/MovieSceneTools/Private/TrackEditors/TransformTrackEditor.cpp
+++ b/Engine/Source/Editor/MovieSceneTools/Private/TrackEditors/TransformTrackEditor.cpp
@@ -260,8 +260,8 @@ void F3DTransformTrackEditor::BuildTrackContextMenu( FMenuBuilder& MenuBuilder,
 	}
 	UMovieScene3DTransformTrack* TransformTrack = Cast<UMovieScene3DTransformTrack>( Track );
 	MenuBuilder.AddMenuEntry(
-		NSLOCTEXT( "Sequencer", "PasteMatineeTrack", "Paste Matinee Move Track" ),
-		NSLOCTEXT( "Sequencer", "PasteMatineeTrackTooltip", "Pastes keys from a Matinee move track into this track." ),
+		NSLOCTEXT( "Sequencer", "PasteMatineeMoveTrack", "Paste Matinee Move Track" ),
+		NSLOCTEXT( "Sequencer", "PasteMatineeMoveTrackTooltip", "Pastes keys from a Matinee move track into this track." ),
 		FSlateIcon(),
 		FUIAction(
 			FExecuteAction::CreateStatic( &CopyInterpMoveTrack, GetSequencer().ToSharedRef(), MoveTrack, TransformTrack ),
diff --git a/Engine/Source/Editor/PListEditor/Public/SPlistEditor.h b/Engine/Source/Editor/PListEditor/Public/SPlistEditor.h
--- a/Engine/Source/Editor/PListEditor/Public/SPlistEditor.h
+++ b/Engine/Source/Editor/PListEditor/Public/SPlistEditor.h
@@ -11,7 +11,7 @@
 #include "../Private/PListNodeFile.h"
 #include "../Private/PListNodeString.h"
 
-#define LOCTEXT_NAMESPACE ""
+#define LOCTEXT_NAMESPACE "PListEditorCommands"
 
 // Type of notification to spawn
 enum ENTF_Types
diff --git a/Engine/Source/Editor/Persona/Private/AnimGraphNodeDetails.cpp b/Engine/Source/Editor/Persona/Private/AnimGraphNodeDetails.cpp
--- a/Engine/Source/Editor/Persona/Private/AnimGraphNodeDetails.cpp
+++ b/Engine/Source/Editor/Persona/Private/AnimGraphNodeDetails.cpp
@@ -709,12 +709,12 @@ TSharedRef<SWidget> SParentPlayerTreeRow::GenerateWidgetForColumn(const FName& C
 
 bool SParentPlayerTreeRow::OnShouldFilterAsset(const FAssetData& AssetData)
 {
-	const FString* SkeletonName = AssetData.TagsAndValues.Find(TEXT("Skeleton"));
+	const FString SkeletonName = AssetData.GetTagValueRef<FString>("Skeleton");
 
-	if(SkeletonName)
+	if(!SkeletonName.IsEmpty())
 	{
 		USkeleton* CurrentSkeleton = GraphNode->GetAnimBlueprint()->TargetSkeleton;
-		if(*SkeletonName == FString::Printf(TEXT("%s'%s'"), *CurrentSkeleton->GetClass()->GetName(), *CurrentSkeleton->GetPathName()))
+		if(SkeletonName == FString::Printf(TEXT("%s'%s'"), *CurrentSkeleton->GetClass()->GetName(), *CurrentSkeleton->GetPathName()))
 		{
 			return false;
 		}
diff --git a/Engine/Source/Editor/Persona/Private/AnimInstanceDetails.cpp b/Engine/Source/Editor/Persona/Private/AnimInstanceDetails.cpp
--- a/Engine/Source/Editor/Persona/Private/AnimInstanceDetails.cpp
+++ b/Engine/Source/Editor/Persona/Private/AnimInstanceDetails.cpp
@@ -117,6 +117,6 @@ TSharedRef<SWidget> FAnimInstanceDetails::CreateFilteredObjectPropertyWidget(UPr
 
 bool FAnimInstanceDetails::OnShouldFilterAnimAsset(const FAssetData& AssetData) const
 {
-	const FString* SkeletonName = AssetData.TagsAndValues.Find(TEXT("Skeleton"));
-	return *SkeletonName != TargetSkeletonName;
+	const FString SkeletonName = AssetData.GetTagValueRef<FString>("Skeleton");
+	return SkeletonName != TargetSkeletonName;
 }
diff --git a/Engine/Source/Editor/Persona/Private/AnimViewportLODCommands.cpp b/Engine/Source/Editor/Persona/Private/AnimViewportLODCommands.cpp
--- a/Engine/Source/Editor/Persona/Private/AnimViewportLODCommands.cpp
+++ b/Engine/Source/Editor/Persona/Private/AnimViewportLODCommands.cpp
@@ -3,7 +3,7 @@
 #include "PersonaPrivatePCH.h"
 #include "AnimViewportLODCommands.h"
 
-#define LOCTEXT_NAMESPACE ""
+#define LOCTEXT_NAMESPACE "AnimViewportLODCommands"
 
 void FAnimViewportLODCommands::RegisterCommands()
 {
diff --git a/Engine/Source/Editor/Persona/Private/AnimViewportMenuCommands.cpp b/Engine/Source/Editor/Persona/Private/AnimViewportMenuCommands.cpp
--- a/Engine/Source/Editor/Persona/Private/AnimViewportMenuCommands.cpp
+++ b/Engine/Source/Editor/Persona/Private/AnimViewportMenuCommands.cpp
@@ -3,7 +3,7 @@
 #include "PersonaPrivatePCH.h"
 #include "AnimViewportMenuCommands.h"
 
-#define LOCTEXT_NAMESPACE ""
+#define LOCTEXT_NAMESPACE "AnimViewportMenuCommands"
 
 void FAnimViewportMenuCommands::RegisterCommands()
 {
diff --git a/Engine/Source/Editor/Persona/Private/AnimViewportPlaybackCommands.cpp b/Engine/Source/Editor/Persona/Private/AnimViewportPlaybackCommands.cpp
--- a/Engine/Source/Editor/Persona/Private/AnimViewportPlaybackCommands.cpp
+++ b/Engine/Source/Editor/Persona/Private/AnimViewportPlaybackCommands.cpp
@@ -4,7 +4,7 @@
 #include "AnimViewportPlaybackCommands.h"
 #include "SAnimationEditorViewport.h"
 
-#define LOCTEXT_NAMESPACE ""
+#define LOCTEXT_NAMESPACE "AnimViewportPlaybackCommands"
 
 FAnimViewportPlaybackCommands::FAnimViewportPlaybackCommands() : TCommands<FAnimViewportPlaybackCommands>
 	(
diff --git a/Engine/Source/Editor/Persona/Private/AnimViewportShowCommands.cpp b/Engine/Source/Editor/Persona/Private/AnimViewportShowCommands.cpp
--- a/Engine/Source/Editor/Persona/Private/AnimViewportShowCommands.cpp
+++ b/Engine/Source/Editor/Persona/Private/AnimViewportShowCommands.cpp
@@ -3,7 +3,7 @@
 #include "PersonaPrivatePCH.h"
 #include "AnimViewportShowCommands.h"
 
-#define LOCTEXT_NAMESPACE ""
+#define LOCTEXT_NAMESPACE "AnimViewportShowCommands"
 
 void FAnimViewportShowCommands::RegisterCommands()
 {
diff --git a/Engine/Source/Editor/Persona/Private/MeshMode/SAdditionalMeshesEditor.cpp b/Engine/Source/Editor/Persona/Private/MeshMode/SAdditionalMeshesEditor.cpp
--- a/Engine/Source/Editor/Persona/Private/MeshMode/SAdditionalMeshesEditor.cpp
+++ b/Engine/Source/Editor/Persona/Private/MeshMode/SAdditionalMeshesEditor.cpp
@@ -271,14 +271,11 @@ void SAdditionalMeshesEditor::OnAssetSelectedFromPicker(const FAssetData& AssetD
 
 bool SAdditionalMeshesEditor::ShouldFilterAssetBasedOnSkeleton(const FAssetData& AssetData)
 {
-	const FString* SkeletonName = AssetData.TagsAndValues.Find(TEXT("Skeleton"));
+	const FString SkeletonName = AssetData.GetTagValueRef<FString>("Skeleton");
 
-	if ( SkeletonName )
+	if ( !SkeletonName.IsEmpty() && SkeletonName == SkeletonNameAssetFilter )
 	{
-		if ( (*SkeletonName) == SkeletonNameAssetFilter )
-		{
-			return false;
-		}
+		return false;
 	}
 
 	return true;
diff --git a/Engine/Source/Editor/Persona/Private/PersonaCommands.cpp b/Engine/Source/Editor/Persona/Private/PersonaCommands.cpp
--- a/Engine/Source/Editor/Persona/Private/PersonaCommands.cpp
+++ b/Engine/Source/Editor/Persona/Private/PersonaCommands.cpp
@@ -4,7 +4,7 @@
 #include "PersonaPrivatePCH.h"
 #include "PersonaCommands.h"
 
-#define LOCTEXT_NAMESPACE ""
+#define LOCTEXT_NAMESPACE "PersonaCommands"
 
 void FPersonaCommands::RegisterCommands()
 {
diff --git a/Engine/Source/Editor/Persona/Private/PersonaModule.cpp b/Engine/Source/Editor/Persona/Private/PersonaModule.cpp
--- a/Engine/Source/Editor/Persona/Private/PersonaModule.cpp
+++ b/Engine/Source/Editor/Persona/Private/PersonaModule.cpp
@@ -38,7 +38,7 @@ void FPersonaModule::StartupModule()
 
 		for (int32 AssetIndex = 0; AssetIndex < AssetData.Num(); ++AssetIndex)
 		{
-			FString TagValue = AssetData[ AssetIndex ].TagsAndValues.FindRef(BPParentClassName);
+			FString TagValue = AssetData[ AssetIndex ].GetTagValueRef<FString>(BPParentClassName);
 			if (TagValue == BPAnimNotify)
 			{
 				FString BlueprintPath = AssetData[AssetIndex].ObjectPath.ToString();
diff --git a/Engine/Source/Editor/Persona/Private/SAnimNotifyPanel.cpp b/Engine/Source/Editor/Persona/Private/SAnimNotifyPanel.cpp
--- a/Engine/Source/Editor/Persona/Private/SAnimNotifyPanel.cpp
+++ b/Engine/Source/Editor/Persona/Private/SAnimNotifyPanel.cpp
@@ -4736,11 +4736,11 @@ void SAnimNotifyPanel::OnGetNotifyBlueprintData(TArray<FAssetData>& OutNotifyDat
 		for(int32 AssetIndex = 0; AssetIndex < AssetDataList.Num(); ++AssetIndex)
 		{
 			FAssetData& AssetData = AssetDataList[AssetIndex];
-			FString TagValue = AssetData.TagsAndValues.FindRef(BPParentClassName);
+			FString TagValue = AssetData.GetTagValueRef<FString>(BPParentClassName);
 
 			if(InOutAllowedClassNames->Contains(TagValue))
 			{
-				FString GenClass = AssetData.TagsAndValues.FindRef(BPGenClassName);
+				FString GenClass = AssetData.GetTagValueRef<FString>(BPGenClassName);
 
 				if(!OutNotifyData.Contains(AssetData))
 				{
diff --git a/Engine/Source/Editor/Persona/Private/SAnimSegmentsPanel.cpp b/Engine/Source/Editor/Persona/Private/SAnimSegmentsPanel.cpp
--- a/Engine/Source/Editor/Persona/Private/SAnimSegmentsPanel.cpp
+++ b/Engine/Source/Editor/Persona/Private/SAnimSegmentsPanel.cpp
@@ -170,7 +170,7 @@ FText SAnimSegmentsPanel::GetAnimSegmentDetailedInfo(int32 AnimSegmentIndex) con
 			}
 			else
 			{
-				return FText::Format(LOCTEXT("AnimSegmentPanel_GetAnimSegmentDetailedInfoFmt", "{0} {1} - ERROR: Recursive Reference Found"), FText::FromString(Anim->GetName()), FText::AsNumber(AnimSegment.GetLength(), &FormatOptions));  
+				return FText::Format(LOCTEXT("AnimSegmentPanel_GetAnimSegmentDetailedInfoFmt_Error_RecursiveReference", "{0} {1} - ERROR: Recursive Reference Found"), FText::FromString(Anim->GetName()), FText::AsNumber(AnimSegment.GetLength(), &FormatOptions));  
 			}
 			
 		}
diff --git a/Engine/Source/Editor/Persona/Private/SAnimationSequenceBrowser.cpp b/Engine/Source/Editor/Persona/Private/SAnimationSequenceBrowser.cpp
--- a/Engine/Source/Editor/Persona/Private/SAnimationSequenceBrowser.cpp
+++ b/Engine/Source/Editor/Persona/Private/SAnimationSequenceBrowser.cpp
@@ -34,7 +34,7 @@ public:
 	// IFilter implementation
 	virtual bool PassesFilter(FAssetFilterType InItem) const override
 	{
-		const FString TagValue = InItem.TagsAndValues.FindRef(GET_MEMBER_NAME_CHECKED(UAnimSequence, AdditiveAnimType));
+		const FString TagValue = InItem.GetTagValueRef<FString>(GET_MEMBER_NAME_CHECKED(UAnimSequence, AdditiveAnimType));
 		return !TagValue.IsEmpty() && !TagValue.Equals(TEXT("AAT_None"));
 	}
 };
diff --git a/Engine/Source/Editor/Persona/Private/SPersonaToolbar.cpp b/Engine/Source/Editor/Persona/Private/SPersonaToolbar.cpp
--- a/Engine/Source/Editor/Persona/Private/SPersonaToolbar.cpp
+++ b/Engine/Source/Editor/Persona/Private/SPersonaToolbar.cpp
@@ -219,12 +219,12 @@ void FPersonaToolbar::FillPersonaModeToolbar(FToolBarBuilder& ParentToolbarBuild
 
 bool FPersonaToolbar::ShouldFilterAssetBasedOnSkeleton(const FAssetData& AssetData)
 {
-	const FString* SkeletonName = AssetData.TagsAndValues.Find(TEXT("Skeleton"));
+	const FString SkeletonName = AssetData.GetTagValueRef<FString>("Skeleton");
 
-	if ( SkeletonName )
+	if ( !SkeletonName.IsEmpty() )
 	{
 		USkeleton* Skeleton = Persona.Pin()->GetSkeleton();
-		if ( (*SkeletonName) == FString::Printf(TEXT("%s'%s'"), *Skeleton->GetClass()->GetName(), *Skeleton->GetPathName()) )
+		if ( SkeletonName == FString::Printf(TEXT("%s'%s'"), *Skeleton->GetClass()->GetName(), *Skeleton->GetPathName()) )
 		{
 			return false;
 		}
diff --git a/Engine/Source/Editor/Persona/Private/SSkeletonAnimNotifies.cpp b/Engine/Source/Editor/Persona/Private/SSkeletonAnimNotifies.cpp
--- a/Engine/Source/Editor/Persona/Private/SSkeletonAnimNotifies.cpp
+++ b/Engine/Source/Editor/Persona/Private/SSkeletonAnimNotifies.cpp
@@ -471,7 +471,7 @@ void SSkeletonAnimNotifies::GetCompatibleAnimSequences(TArray<class FAssetData>&
 	for( int32 AssetIndex = 0; AssetIndex < AssetDataList.Num(); ++AssetIndex )
 	{
 		const FAssetData& PossibleAnimSequence = AssetDataList[AssetIndex];
-		if( SkeletonExportName == PossibleAnimSequence.TagsAndValues.FindRef("Skeleton") )
+		if( SkeletonExportName == PossibleAnimSequence.GetTagValueRef<FString>("Skeleton") )
 		{
 			OutAssets.Add( PossibleAnimSequence );
 		}
diff --git a/Engine/Source/Editor/Persona/Private/SSkeletonSlotNames.cpp b/Engine/Source/Editor/Persona/Private/SSkeletonSlotNames.cpp
--- a/Engine/Source/Editor/Persona/Private/SSkeletonSlotNames.cpp
+++ b/Engine/Source/Editor/Persona/Private/SSkeletonSlotNames.cpp
@@ -495,7 +495,7 @@ void SSkeletonSlotNames::GetCompatibleAnimBlueprints( TArray<FAssetData>& OutAss
 
 	for(FAssetData& Data : AssetDataList)
 	{
-		FString AssetSkeleton = Data.TagsAndValues.FindRef("TargetSkeleton");
+		const FString AssetSkeleton = Data.GetTagValueRef<FString>("TargetSkeleton");
 		if(AssetSkeleton == SkeletonExportName)
 		{
 			OutAssets.Add(Data);
@@ -580,7 +580,7 @@ void SSkeletonSlotNames::GetCompatibleAnimMontages(TArray<class FAssetData>& Out
 	for( int32 AssetIndex = 0; AssetIndex < AssetDataList.Num(); ++AssetIndex )
 	{
 		const FAssetData& PossibleAnimMontage = AssetDataList[AssetIndex];
-		if( SkeletonExportName == PossibleAnimMontage.TagsAndValues.FindRef("Skeleton") )
+		if( SkeletonExportName == PossibleAnimMontage.GetTagValueRef<FString>("Skeleton") )
 		{
 			OutAssets.Add( PossibleAnimMontage );
 		}
diff --git a/Engine/Source/Editor/Persona/Private/SSkeletonSmartNameManager.cpp b/Engine/Source/Editor/Persona/Private/SSkeletonSmartNameManager.cpp
--- a/Engine/Source/Editor/Persona/Private/SSkeletonSmartNameManager.cpp
+++ b/Engine/Source/Editor/Persona/Private/SSkeletonSmartNameManager.cpp
@@ -388,13 +388,11 @@ void SCurveNameManager::OnDeleteNameClicked()
 		FAssetData& Data = AnimationAssets[Idx];
 		bool bRemove = true;
 
-		const FString* SkeletonDataTag = Data.TagsAndValues.Find("Skeleton");
-		if(SkeletonDataTag && *SkeletonDataTag == CurrentSkeletonName)
+		const FString SkeletonDataTag = Data.GetTagValueRef<FString>("Skeleton");
+		if(!SkeletonDataTag.IsEmpty() && SkeletonDataTag == CurrentSkeletonName)
 		{
-			const FString* CurveData = Data.TagsAndValues.Find(USkeleton::CurveTag);
-			FString CurveDataCopy;
-
-			if(!CurveData)
+			FString CurveData;
+			if(!Data.GetTagValue(USkeleton::CurveTag, CurveData))
 			{
 				// This asset is old; it hasn't been loaded since before smartnames were added for curves.
 				// unfortunately the only way to delete curves safely is to load old assets to see if they're
@@ -408,15 +406,11 @@ void SCurveNameManager::OnDeleteNameClicked()
 				{
 					return InTag.Name == USkeleton::CurveTag;
 				});
-				CurveDataCopy = CurveTag->Value;
-			}
-			else
-			{
-				CurveDataCopy = *CurveData;
+				CurveData = CurveTag->Value;
 			}
 
 			TArray<FString> ParsedStringUids;
-			CurveDataCopy.ParseIntoArray(ParsedStringUids, *USkeleton::CurveTagDelimiter, true);
+			CurveData.ParseIntoArray(ParsedStringUids, *USkeleton::CurveTagDelimiter, true);
 
 			for(const FString& UidString : ParsedStringUids)
 			{
diff --git a/Engine/Source/Editor/Persona/Private/SSkeletonTree.cpp b/Engine/Source/Editor/Persona/Private/SSkeletonTree.cpp
--- a/Engine/Source/Editor/Persona/Private/SSkeletonTree.cpp
+++ b/Engine/Source/Editor/Persona/Private/SSkeletonTree.cpp
@@ -18,7 +18,7 @@
 #include "Editor/ContentBrowser/Public/ContentBrowserModule.h"
 #include "ComponentAssetBroker.h"
 
-#include "ClassIconFinder.h"
+#include "SlateIconFinder.h"
 
 #include "Editor/UnrealEd/Public/AssetNotifications.h"
 
@@ -777,7 +777,7 @@ TSharedRef<ITableRow> FDisplayedAttachedAssetInfo::MakeTreeRowWidget(
 void FDisplayedAttachedAssetInfo::GenerateWidgetForNameColumn( TSharedPtr< SHorizontalBox > Box, FText& FilterText, FIsSelected InIsSelected )
 {
 	UActorFactory* ActorFactory = FActorFactoryAssetProxy::GetFactoryForAssetObject( Asset );
-	const FSlateBrush* IconBrush = FClassIconFinder::FindIconForClass( ActorFactory->GetDefaultActorClass( FAssetData() ) );		
+	const FSlateBrush* IconBrush = FSlateIconFinder::FindIconBrushForClass(ActorFactory->GetDefaultActorClass(FAssetData()));
 	
 	Box->AddSlot()
 		.AutoWidth()
diff --git a/Engine/Source/Editor/Persona/Private/SkeletonTreeCommands.cpp b/Engine/Source/Editor/Persona/Private/SkeletonTreeCommands.cpp
--- a/Engine/Source/Editor/Persona/Private/SkeletonTreeCommands.cpp
+++ b/Engine/Source/Editor/Persona/Private/SkeletonTreeCommands.cpp
@@ -3,7 +3,7 @@
 #include "PersonaPrivatePCH.h"
 #include "SkeletonTreeCommands.h"
 
-#define LOCTEXT_NAMESPACE ""
+#define LOCTEXT_NAMESPACE "SkeletonTreeCommands"
 
 void FSkeletonTreeCommands::RegisterCommands()
 {
diff --git a/Engine/Source/Editor/PhAT/Private/PhAT.cpp b/Engine/Source/Editor/PhAT/Private/PhAT.cpp
--- a/Engine/Source/Editor/PhAT/Private/PhAT.cpp
+++ b/Engine/Source/Editor/PhAT/Private/PhAT.cpp
@@ -1778,9 +1778,9 @@ UObject* FPhAT::GetSelectedAnimation() const
 bool FPhAT::ShouldFilterAssetBasedOnSkeleton( const FAssetData& AssetData )
 {
 	// @TODO This is a duplicate of FPersona::ShouldFilterAssetBasedOnSkeleton(), but should go away once PhAT is integrated with Persona
-	const FString* SkeletonName = AssetData.TagsAndValues.Find(TEXT("Skeleton"));
+	const FString SkeletonName = AssetData.GetTagValueRef<FString>("Skeleton");
 
-	if ( SkeletonName )
+	if ( !SkeletonName.IsEmpty() )
 	{
 		USkeleton* Skeleton = SharedData->EditorSkelMesh->Skeleton;
 
diff --git a/Engine/Source/Editor/PhAT/Private/PhATActions.cpp b/Engine/Source/Editor/PhAT/Private/PhATActions.cpp
--- a/Engine/Source/Editor/PhAT/Private/PhATActions.cpp
+++ b/Engine/Source/Editor/PhAT/Private/PhATActions.cpp
@@ -4,7 +4,7 @@
 #include "PhATActions.h"
 #include "PhAT.h"
 
-#define LOCTEXT_NAMESPACE ""
+#define LOCTEXT_NAMESPACE "PhATCommands"
 
 void FPhATCommands::RegisterCommands()
 {
diff --git a/Engine/Source/Editor/PropertyEditor/Private/PropertyCustomizationHelpers.cpp b/Engine/Source/Editor/PropertyEditor/Private/PropertyCustomizationHelpers.cpp
--- a/Engine/Source/Editor/PropertyEditor/Private/PropertyCustomizationHelpers.cpp
+++ b/Engine/Source/Editor/PropertyEditor/Private/PropertyCustomizationHelpers.cpp
@@ -307,7 +307,7 @@ namespace PropertyCustomizationHelpers
 			if (Class->IsChildOf(UFactory::StaticClass()) && !Class->HasAnyClassFlags(CLASS_Abstract))
 			{
 				UFactory* Factory = Class->GetDefaultObject<UFactory>();
-				if (Factory->CanCreateNew() && ensure(!Factory->GetDisplayName().IsEmpty()))
+				if (Factory->ShouldShowInNewMenu() && ensure(!Factory->GetDisplayName().IsEmpty()))
 				{
 					UClass* SupportedClass = Factory->GetSupportedClass();
 					if (SupportedClass != nullptr && Classes.ContainsByPredicate([=](const UClass* InClass) { return SupportedClass->IsChildOf(InClass); }))
@@ -318,6 +318,11 @@ namespace PropertyCustomizationHelpers
 			}
 		}
 
+		Factories.Sort([](UFactory& A, UFactory& B) -> bool
+		{
+			return A.GetDisplayName().CompareToCaseIgnored(B.GetDisplayName()) < 0;
+		});
+
 		return Factories;
 	}
 }
diff --git a/Engine/Source/Editor/PropertyEditor/Private/SDetailNameArea.cpp b/Engine/Source/Editor/PropertyEditor/Private/SDetailNameArea.cpp
--- a/Engine/Source/Editor/PropertyEditor/Private/SDetailNameArea.cpp
+++ b/Engine/Source/Editor/PropertyEditor/Private/SDetailNameArea.cpp
@@ -3,7 +3,7 @@
 
 #include "PropertyEditorPrivatePCH.h"
 #include "SDetailNameArea.h"
-#include "ClassIconFinder.h"
+#include "SlateIconFinder.h"
 #include "Editor/EditorWidgets/Public/EditorWidgets.h"
 #include "EditorClassUtils.h"
 
@@ -103,7 +103,7 @@ TSharedRef< SWidget > SDetailNameArea::BuildObjectNameArea( const TArray< TWeakO
 	if (BaseClass)
 	{
 		// Get selection icon based on actor(s) classes and add before the selection label
-		const FSlateBrush* ActorIcon = FClassIconFinder::FindIconForClass(BaseClass);
+		const FSlateBrush* ActorIcon = FSlateIconFinder::FindIconBrushForClass(BaseClass);
 
 		ObjectNameArea->AddSlot()
 		.AutoWidth()
diff --git a/Engine/Source/Editor/PropertyEditor/Private/UserInterface/PropertyEditor/SPropertyEditorEditInline.cpp b/Engine/Source/Editor/PropertyEditor/Private/UserInterface/PropertyEditor/SPropertyEditorEditInline.cpp
--- a/Engine/Source/Editor/PropertyEditor/Private/UserInterface/PropertyEditor/SPropertyEditorEditInline.cpp
+++ b/Engine/Source/Editor/PropertyEditor/Private/UserInterface/PropertyEditor/SPropertyEditorEditInline.cpp
@@ -10,7 +10,7 @@
 #include "SPropertyComboBox.h"
 #include "Editor/ClassViewer/Public/ClassViewerModule.h"
 #include "Editor/ClassViewer/Public/ClassViewerFilter.h"
-#include "ClassIconFinder.h"
+#include "SlateIconFinder.h"
 
 class FPropertyEditorInlineClassFilter : public IClassViewerFilter
 {
@@ -116,7 +116,7 @@ const FSlateBrush* SPropertyEditorEditInline::GetDisplayValueIcon() const
 	FPropertyAccess::Result Result = PropertyEditor->GetPropertyHandle()->GetValue( CurrentValue );
 	if( Result == FPropertyAccess::Success && CurrentValue != nullptr )
 	{
-		return FClassIconFinder::FindIconForClass(CurrentValue->GetClass());
+		return FSlateIconFinder::FindIconBrushForClass(CurrentValue->GetClass());
 	}
 
 	return nullptr;
diff --git a/Engine/Source/Editor/PropertyEditor/Private/UserInterface/PropertyEditor/SPropertyEditorText.cpp b/Engine/Source/Editor/PropertyEditor/Private/UserInterface/PropertyEditor/SPropertyEditorText.cpp
--- a/Engine/Source/Editor/PropertyEditor/Private/UserInterface/PropertyEditor/SPropertyEditorText.cpp
+++ b/Engine/Source/Editor/PropertyEditor/Private/UserInterface/PropertyEditor/SPropertyEditorText.cpp
@@ -13,15 +13,12 @@ void SPropertyEditorText::Construct( const FArguments& InArgs, const TSharedRef<
 {
 	PropertyEditor = InPropertyEditor;
 
-	const UProperty* Property = InPropertyEditor->GetProperty();
+	bIsFNameProperty = InPropertyEditor->PropertyIsA(UNameProperty::StaticClass());
+	bIsMultiLine = InPropertyEditor->GetPropertyHandle()->GetMetaDataProperty()->GetBoolMetaData("MultiLine");
 
+	const bool bIsPassword = InPropertyEditor->GetPropertyHandle()->GetMetaDataProperty()->GetBoolMetaData("PasswordField");
+	
 	TSharedPtr<SHorizontalBox> HorizontalBox;
-
-	bIsFNameProperty = Property->IsA<UNameProperty>();
-
-	bool bIsPassword = Property->GetBoolMetaData("PasswordField");
-
-	bIsMultiLine = Property->GetBoolMetaData("MultiLine");
 	if(bIsMultiLine)
 	{
 		ChildSlot
diff --git a/Engine/Source/Editor/PropertyEditor/Private/UserInterface/PropertyEditor/SPropertyMenuAssetPicker.cpp b/Engine/Source/Editor/PropertyEditor/Private/UserInterface/PropertyEditor/SPropertyMenuAssetPicker.cpp
--- a/Engine/Source/Editor/PropertyEditor/Private/UserInterface/PropertyEditor/SPropertyMenuAssetPicker.cpp
+++ b/Engine/Source/Editor/PropertyEditor/Private/UserInterface/PropertyEditor/SPropertyMenuAssetPicker.cpp
@@ -8,7 +8,7 @@
 #include "ContentBrowserModule.h"
 #include "PropertyEditorAssetConstants.h"
 #include "EditorStyleSet.h"
-#include "ClassIconFinder.h"
+#include "SlateIconFinder.h"
 
 #define LOCTEXT_NAMESPACE "PropertyEditor"
 
@@ -35,7 +35,7 @@ void SPropertyMenuAssetPicker::Construct( const FArguments& InArgs )
 				MenuBuilder.AddMenuEntry(
 					Factory->GetDisplayName(),
 					Factory->GetToolTip(),
-					FSlateIcon(FEditorStyle::GetStyleSetName(), FClassIconFinder::FindIconNameForClass(Factory->GetSupportedClass())),
+					FSlateIconFinder::FindIconForClass(Factory->GetSupportedClass()),
 					FUIAction(FExecuteAction::CreateSP(this, &SPropertyMenuAssetPicker::OnCreateNewAssetSelected, FactoryPtr))
 					);
 			}
diff --git a/Engine/Source/Editor/ReferenceViewer/Private/ReferenceViewerActions.cpp b/Engine/Source/Editor/ReferenceViewer/Private/ReferenceViewerActions.cpp
--- a/Engine/Source/Editor/ReferenceViewer/Private/ReferenceViewerActions.cpp
+++ b/Engine/Source/Editor/ReferenceViewer/Private/ReferenceViewerActions.cpp
@@ -6,7 +6,7 @@
 //////////////////////////////////////////////////////////////////////////
 // FReferenceViewerActions
 
-#define LOCTEXT_NAMESPACE ""
+#define LOCTEXT_NAMESPACE "ReferenceViewerActions"
 
 void FReferenceViewerActions::RegisterCommands()
 {
diff --git a/Engine/Source/Editor/SceneOutliner/Private/SSceneOutliner.cpp b/Engine/Source/Editor/SceneOutliner/Private/SSceneOutliner.cpp
--- a/Engine/Source/Editor/SceneOutliner/Private/SSceneOutliner.cpp
+++ b/Engine/Source/Editor/SceneOutliner/Private/SSceneOutliner.cpp
@@ -967,6 +967,7 @@ namespace SceneOutliner
 
 		PendingOperations.Empty();
 		TreeItemMap.Reset();
+		PendingTreeItemMap.Empty();
 
 		RootTreeItems.Empty();
 	}
@@ -1108,6 +1109,8 @@ namespace SceneOutliner
 	{
 		const auto ItemID = Item->GetID();
 
+		PendingTreeItemMap.Remove(ItemID);
+
 		// If a tree item already exists that represents the same data, bail
 		if (TreeItemMap.Find(ItemID))
 		{
@@ -1978,7 +1981,7 @@ namespace SceneOutliner
 		{
 			if( InActor && SharedData->RepresentingWorld == InActor->GetWorld() && IsActorDisplayable(InActor) )
 			{
-				if (!TreeItemMap.Find(InActor))
+				if (!TreeItemMap.Find(InActor) && !PendingTreeItemMap.Find(InActor))
 				{
 					// Update the total actor count that match the filters
 					if (Filters->PassesAllFilters(FActorTreeItem(InActor)))
@@ -1999,8 +2002,13 @@ namespace SceneOutliner
 			if( InActor && SharedData->RepresentingWorld == InActor->GetWorld() )
 			{
 				ApplicableActors.Remove(InActor);
+				auto* ItemPtr = TreeItemMap.Find(InActor);
+				if (!ItemPtr)
+				{
+					ItemPtr = PendingTreeItemMap.Find(InActor);
+				}
 
-				if (auto* ItemPtr = TreeItemMap.Find(InActor))
+				if (ItemPtr)
 				{
 					PendingOperations.Emplace(FPendingTreeOperation::Removed, ItemPtr->ToSharedRef());
 					Refresh();
@@ -2386,6 +2394,8 @@ namespace SceneOutliner
 						const EEditAction::Type CanProcess = ActiveModes[ModeIndex]->GetActionEditDelete();
 						if (CanProcess == EEditAction::Process)
 						{
+							// We don't consider the return value here, as false is assumed to mean there was an internal error processing delete, not that it should defer to other modes/default behaviour
+							ActiveModes[ModeIndex]->ProcessEditDelete();
 							return FReply::Handled();
 						}
 						else if (CanProcess == EEditAction::Halt)
diff --git a/Engine/Source/Editor/SceneOutliner/Private/SSceneOutliner.h b/Engine/Source/Editor/SceneOutliner/Private/SSceneOutliner.h
--- a/Engine/Source/Editor/SceneOutliner/Private/SSceneOutliner.h
+++ b/Engine/Source/Editor/SceneOutliner/Private/SSceneOutliner.h
@@ -148,6 +148,7 @@ namespace SceneOutliner
 			{
 				FTreeItemRef NewItem = MakeShareable(new TreeItemType(Data));
 				PendingOperations.Emplace(FPendingTreeOperation::Added, NewItem);
+				PendingTreeItemMap.Add(NewItem->GetID(), NewItem);
 				Refresh();
 			}
 		}
@@ -429,6 +430,9 @@ namespace SceneOutliner
 		/** A map of all items we have in the tree */
 		FTreeItemMap TreeItemMap;
 
+		/** Pending tree items that are yet to be added the tree */
+		FTreeItemMap PendingTreeItemMap;
+
 		/** Root level tree items */
 		TArray<FTreeItemPtr> RootTreeItems;
 
diff --git a/Engine/Source/Editor/SceneOutliner/Private/SceneOutlinerItemLabelColumn.cpp b/Engine/Source/Editor/SceneOutliner/Private/SceneOutlinerItemLabelColumn.cpp
--- a/Engine/Source/Editor/SceneOutliner/Private/SceneOutlinerItemLabelColumn.cpp
+++ b/Engine/Source/Editor/SceneOutliner/Private/SceneOutlinerItemLabelColumn.cpp
@@ -365,7 +365,7 @@ struct SWorldTreeLabel : FCommonLabelData, public SCompoundWidget
 				.HeightOverride(FDefaultTreeItemMetrics::IconSize())
 				[
 					SNew(SImage)
-					.Image(FClassIconFinder::FindIconForClass(UWorld::StaticClass()))
+					.Image(FSlateIconFinder::FindIconBrushForClass(UWorld::StaticClass()))
 					.ToolTipText(LOCTEXT("WorldIcon_Tooltip", "World"))
 				]
 			]
diff --git a/Engine/Source/Editor/SequenceRecorder/Private/SequenceRecorder.cpp b/Engine/Source/Editor/SequenceRecorder/Private/SequenceRecorder.cpp
--- a/Engine/Source/Editor/SequenceRecorder/Private/SequenceRecorder.cpp
+++ b/Engine/Source/Editor/SequenceRecorder/Private/SequenceRecorder.cpp
@@ -340,7 +340,7 @@ void FSequenceRecorder::DrawDebug(UCanvas* InCanvas, APlayerController* InPlayer
 		}
 		else
 		{
-			Text = LOCTEXT("RecordingIndicator", "Sequence Recording Pending");
+			Text = LOCTEXT("RecordingIndicatorPending", "Sequence Recording Pending");
 		}
 
 		const FVector2D IconSize(32.0f, 32.0f);
diff --git a/Engine/Source/Editor/Sequencer/Private/DisplayNodes/SequencerObjectBindingNode.cpp b/Engine/Source/Editor/Sequencer/Private/DisplayNodes/SequencerObjectBindingNode.cpp
--- a/Engine/Source/Editor/Sequencer/Private/DisplayNodes/SequencerObjectBindingNode.cpp
+++ b/Engine/Source/Editor/Sequencer/Private/DisplayNodes/SequencerObjectBindingNode.cpp
@@ -9,7 +9,7 @@
 #include "ObjectEditorUtils.h"
 #include "AssetEditorManager.h"
 #include "SequencerUtilities.h"
-#include "ClassIconFinder.h"
+#include "SlateIconFinder.h"
 #include "SequencerCommands.h"
 
 #define LOCTEXT_NAMESPACE "FObjectBindingNode"
@@ -273,7 +273,7 @@ FText FSequencerObjectBindingNode::GetDisplayNameToolTipText() const
 
 const FSlateBrush* FSequencerObjectBindingNode::GetIconBrush() const
 {
-	return FClassIconFinder::FindIconForClass(GetClassForObjectBinding());
+	return FSlateIconFinder::FindIconBrushForClass(GetClassForObjectBinding());
 }
 
 const FSlateBrush* FSequencerObjectBindingNode::GetIconOverlayBrush() const
diff --git a/Engine/Source/Editor/Sequencer/Private/SequencerCommands.cpp b/Engine/Source/Editor/Sequencer/Private/SequencerCommands.cpp
--- a/Engine/Source/Editor/Sequencer/Private/SequencerCommands.cpp
+++ b/Engine/Source/Editor/Sequencer/Private/SequencerCommands.cpp
@@ -3,7 +3,7 @@
 #include "SequencerPrivatePCH.h"
 #include "SequencerCommands.h"
 
-#define LOCTEXT_NAMESPACE ""
+#define LOCTEXT_NAMESPACE "SequencerCommands"
 
 void FSequencerCommands::RegisterCommands()
 {
diff --git a/Engine/Source/Editor/SizeMap/SizeMapActions.cpp b/Engine/Source/Editor/SizeMap/SizeMapActions.cpp
--- a/Engine/Source/Editor/SizeMap/SizeMapActions.cpp
+++ b/Engine/Source/Editor/SizeMap/SizeMapActions.cpp
@@ -3,7 +3,7 @@
 #include "SizeMapModule.h"
 #include "SizeMapActions.h"
 
-#define LOCTEXT_NAMESPACE ""
+#define LOCTEXT_NAMESPACE "SizeMapActions"
 
 void FSizeMapActions::RegisterCommands()
 {
diff --git a/Engine/Source/Editor/StaticMeshEditor/Private/StaticMeshEditorActions.cpp b/Engine/Source/Editor/StaticMeshEditor/Private/StaticMeshEditorActions.cpp
--- a/Engine/Source/Editor/StaticMeshEditor/Private/StaticMeshEditorActions.cpp
+++ b/Engine/Source/Editor/StaticMeshEditor/Private/StaticMeshEditorActions.cpp
@@ -5,7 +5,7 @@
 #include "StaticMeshEditorActions.h"
 #include "SStaticMeshEditorViewport.h"
 
-#define LOCTEXT_NAMESPACE ""
+#define LOCTEXT_NAMESPACE "StaticMeshEditorCommands"
 
 void FStaticMeshEditorCommands::RegisterCommands()
 {
diff --git a/Engine/Source/Editor/TextureEditor/Private/Models/TextureEditorCommands.cpp b/Engine/Source/Editor/TextureEditor/Private/Models/TextureEditorCommands.cpp
--- a/Engine/Source/Editor/TextureEditor/Private/Models/TextureEditorCommands.cpp
+++ b/Engine/Source/Editor/TextureEditor/Private/Models/TextureEditorCommands.cpp
@@ -2,7 +2,7 @@
 
 #include "TextureEditorPrivatePCH.h"
 
-#define LOCTEXT_NAMESPACE ""
+#define LOCTEXT_NAMESPACE "TextureEditorCommands"
 
 void FTextureEditorCommands::RegisterCommands()
 {
diff --git a/Engine/Source/Editor/TextureEditor/Private/Models/TextureEditorViewportClient.cpp b/Engine/Source/Editor/TextureEditor/Private/Models/TextureEditorViewportClient.cpp
--- a/Engine/Source/Editor/TextureEditor/Private/Models/TextureEditorViewportClient.cpp
+++ b/Engine/Source/Editor/TextureEditor/Private/Models/TextureEditorViewportClient.cpp
@@ -190,7 +190,7 @@ FText FTextureEditorViewportClient::GetDisplayedResolution() const
 	uint32 Height = 1;
 	uint32 Width = 1;
 	TextureEditorPtr.Pin()->CalculateTextureDimensions(Width, Height);
-	return FText::Format( NSLOCTEXT("TextureEditor", "QuickInfo_Displayed", "Displayed: {0}x{1}"), FText::AsNumber( FMath::Max((uint32)1, Width) ), FText::AsNumber( FMath::Max((uint32)1, Height)) );
+	return FText::Format( NSLOCTEXT("TextureEditor", "DisplayedResolution", "Displayed: {0}x{1}"), FText::AsNumber( FMath::Max((uint32)1, Width) ), FText::AsNumber( FMath::Max((uint32)1, Height)) );
 }
 
 
diff --git a/Engine/Source/Editor/TextureEditor/Private/TextureEditorToolkit.cpp b/Engine/Source/Editor/TextureEditor/Private/TextureEditorToolkit.cpp
--- a/Engine/Source/Editor/TextureEditor/Private/TextureEditorToolkit.cpp
+++ b/Engine/Source/Editor/TextureEditor/Private/TextureEditorToolkit.cpp
@@ -337,6 +337,7 @@ void FTextureEditorToolkit::PopulateQuickInfo( )
 	UTextureRenderTarget2D* Texture2DRT = Cast<UTextureRenderTarget2D>(Texture);
 	UTextureRenderTargetCube* TextureCubeRT = Cast<UTextureRenderTargetCube>(Texture);
 	UTextureCube* TextureCube = Cast<UTextureCube>(Texture);
+	UTexture2DDynamic* Texture2DDynamic = Cast<UTexture2DDynamic>(Texture);
 
 	uint32 ImportedWidth = Texture->Source.GetSizeX();
 	uint32 ImportedHeight = Texture->Source.GetSizeY();
@@ -403,6 +404,10 @@ void FTextureEditorToolkit::PopulateQuickInfo( )
 	{
 		TextureFormatIndex = Texture2DRT->GetFormat();
 	}
+	else if (Texture2DDynamic)
+	{
+		TextureFormatIndex = Texture2DDynamic->Format;
+	}
 
 	if (TextureFormatIndex != PF_MAX)
 	{
@@ -422,6 +427,10 @@ void FTextureEditorToolkit::PopulateQuickInfo( )
 	{
 		NumMips = Texture2DRT->GetNumMips();
 	}
+	else if (Texture2DDynamic)
+	{
+		NumMips = Texture2DDynamic->NumMips;
+	}
 
 	NumMipsText->SetText(FText::Format(NSLOCTEXT("TextureEditor", "QuickInfo_NumMips", "Number of Mips: {0}"), FText::AsNumber(NumMips)));
 
diff --git a/Engine/Source/Editor/UATHelper/UATHelper.Build.cs b/Engine/Source/Editor/UATHelper/UATHelper.Build.cs
new file mode 100644
--- /dev/null
+++ b/Engine/Source/Editor/UATHelper/UATHelper.Build.cs
@@ -0,0 +1,23 @@
+// Copyright 1998-2016 Epic Games, Inc. All Rights Reserved.
+
+namespace UnrealBuildTool.Rules
+{
+	public class UATHelper : ModuleRules
+	{
+        public UATHelper(TargetInfo Target)
+		{
+			PrivateDependencyModuleNames.AddRange(
+                new string[] {
+				    "Core",
+				    "CoreUObject",
+				    "Engine",
+                    "InputCore",
+				    "Slate",
+					"SlateCore",
+                    "GameProjectGeneration",
+                    "UnrealEd"
+			    }
+			);
+		}
+	}
+}
\ No newline at end of file
diff --git a/Engine/Source/Editor/UATHelper/UATHelperModule.h b/Engine/Source/Editor/UATHelper/UATHelperModule.h
new file mode 100644
--- /dev/null
+++ b/Engine/Source/Editor/UATHelper/UATHelperModule.h
@@ -0,0 +1,6 @@
+// Copyright 1998-2016 Epic Games, Inc. All Rights Reserved.
+
+#pragma once
+
+#include "UnrealEd.h"
+#include "IUATHelperModule.h"
diff --git a/Engine/Source/Editor/UMGEditor/Private/AssetTypeActions_WidgetBlueprint.cpp b/Engine/Source/Editor/UMGEditor/Private/AssetTypeActions_WidgetBlueprint.cpp
--- a/Engine/Source/Editor/UMGEditor/Private/AssetTypeActions_WidgetBlueprint.cpp
+++ b/Engine/Source/Editor/UMGEditor/Private/AssetTypeActions_WidgetBlueprint.cpp
@@ -46,13 +46,11 @@ UClass* FAssetTypeActions_WidgetBlueprint::GetSupportedClass() const
 
 FText FAssetTypeActions_WidgetBlueprint::GetAssetDescription( const FAssetData& AssetData ) const
 {
-	if ( const FString* pDescription = AssetData.TagsAndValues.Find( GET_MEMBER_NAME_CHECKED( UBlueprint, BlueprintDescription ) ) )
+	FString Description = AssetData.GetTagValueRef<FString>( GET_MEMBER_NAME_CHECKED( UBlueprint, BlueprintDescription ) );
+	if ( !Description.IsEmpty() )
 	{
-		if ( !pDescription->IsEmpty() )
-		{
-			const FString DescriptionStr( *pDescription );
-			return FText::FromString( DescriptionStr.Replace( TEXT( "\\n" ), TEXT( "\n" ) ) );
-		}
+		Description.ReplaceInline( TEXT( "\\n" ), TEXT( "\n" ) );
+		return FText::FromString( MoveTemp(Description) );
 	}
 
 	return FText::GetEmpty();
diff --git a/Engine/Source/Editor/UMGEditor/Private/Customizations/UMGDetailCustomizations.cpp b/Engine/Source/Editor/UMGEditor/Private/Customizations/UMGDetailCustomizations.cpp
--- a/Engine/Source/Editor/UMGEditor/Private/Customizations/UMGDetailCustomizations.cpp
+++ b/Engine/Source/Editor/UMGEditor/Private/Customizations/UMGDetailCustomizations.cpp
@@ -134,7 +134,7 @@ void FBlueprintWidgetCustomization::CreateEventCustomization( IDetailLayoutBuild
 
 void FBlueprintWidgetCustomization::ResetToDefault_RemoveBinding(TSharedRef<IPropertyHandle> PropertyHandle)
 {
-	const FScopedTransaction Transaction(LOCTEXT("BindDelegate", "Remove Binding"));
+	const FScopedTransaction Transaction(LOCTEXT("UnbindDelegate", "Remove Binding"));
 
 	Blueprint->Modify();
 
diff --git a/Engine/Source/Editor/UMGEditor/Private/Designer/DesignerCommands.cpp b/Engine/Source/Editor/UMGEditor/Private/Designer/DesignerCommands.cpp
--- a/Engine/Source/Editor/UMGEditor/Private/Designer/DesignerCommands.cpp
+++ b/Engine/Source/Editor/UMGEditor/Private/Designer/DesignerCommands.cpp
@@ -3,7 +3,7 @@
 #include "UMGEditorPrivatePCH.h"
 #include "DesignerCommands.h"
 
-#define LOCTEXT_NAMESPACE ""
+#define LOCTEXT_NAMESPACE "DesignerCommands"
 
 void FDesignerCommands::RegisterCommands()
 {
diff --git a/Engine/Source/Editor/UMGEditor/Private/Designer/SDesignerView.cpp b/Engine/Source/Editor/UMGEditor/Private/Designer/SDesignerView.cpp
--- a/Engine/Source/Editor/UMGEditor/Private/Designer/SDesignerView.cpp
+++ b/Engine/Source/Editor/UMGEditor/Private/Designer/SDesignerView.cpp
@@ -255,7 +255,7 @@ void SDesignerView::Construct(const FArguments& InArgs, TSharedPtr<FWidgetBluepr
 	Register(MakeShareable(new FUniformGridSlotExtension()));
 	Register(MakeShareable(new FGridSlotExtension()));
 
-	GEditor->OnBlueprintReinstanced().AddRaw(this, &SDesignerView::OnBlueprintReinstanced);
+	GEditor->OnBlueprintReinstanced().AddRaw(this, &SDesignerView::OnPreviewNeedsRecreation);
 
 	BindCommands();
 
@@ -365,6 +365,7 @@ void SDesignerView::Construct(const FArguments& InArgs, TSharedPtr<FWidgetBluepr
 	PinnedBlueprintEditor->OnSelectedWidgetsChanged.AddRaw(this, &SDesignerView::OnEditorSelectionChanged);
 	PinnedBlueprintEditor->OnHoveredWidgetSet.AddRaw(this, &SDesignerView::OnHoveredWidgetSet);
 	PinnedBlueprintEditor->OnHoveredWidgetCleared.AddRaw(this, &SDesignerView::OnHoveredWidgetCleared);
+	PinnedBlueprintEditor->OnWidgetPreviewUpdated.AddRaw(this, &SDesignerView::OnPreviewNeedsRecreation);
 
 	ZoomToFit(/*bInstantZoom*/ true);
 }
@@ -627,6 +628,7 @@ SDesignerView::~SDesignerView()
 		PinnedEditor->OnSelectedWidgetsChanged.RemoveAll(this);
 		PinnedEditor->OnHoveredWidgetSet.RemoveAll(this);
 		PinnedEditor->OnHoveredWidgetCleared.RemoveAll(this);
+		PinnedEditor->OnWidgetPreviewUpdated.RemoveAll(this);
 	}
 
 	if ( GEditor )
@@ -1160,7 +1162,7 @@ void SDesignerView::Register(TSharedRef<FDesignerExtension> Extension)
 	DesignerExtensions.Add(Extension);
 }
 
-void SDesignerView::OnBlueprintReinstanced()
+void SDesignerView::OnPreviewNeedsRecreation()
 {
 	// Because widget blueprints can contain other widget blueprints, the safe thing to do is to have all
 	// designers jettison their previews on the compilation of any widget blueprint.  We do this to prevent
@@ -1676,7 +1678,7 @@ void SDesignerView::UpdatePreviewWidget(bool bForceUpdate)
 
 			// The constraint box for the widget size needs to inside the DPI scaler in order to make sure it too
 			// is sized accurately for the size screen it's on.
-			TSharedRef<SBox> NewPreviewSizeConstraintBox = SAssignNew(PreviewSizeConstraint, SBox)
+			TSharedRef<SBox> NewPreviewSizeConstraintBox = SNew(SBox)
 				.WidthOverride(this, &SDesignerView::GetPreviewSizeWidth)
 				.HeightOverride(this, &SDesignerView::GetPreviewSizeHeight)
 				[
diff --git a/Engine/Source/Editor/UMGEditor/Private/Designer/SDesignerView.h b/Engine/Source/Editor/UMGEditor/Private/Designer/SDesignerView.h
--- a/Engine/Source/Editor/UMGEditor/Private/Designer/SDesignerView.h
+++ b/Engine/Source/Editor/UMGEditor/Private/Designer/SDesignerView.h
@@ -125,8 +125,8 @@ private:
 	/** Gets the blueprint being edited by the designer */
 	UWidgetBlueprint* GetBlueprint() const;
 
-	/** Called whenever the blueprint is recompiled */
-	void OnBlueprintReinstanced();
+	/** Called whenever the blueprint is recompiled or the preview is updated */
+	void OnPreviewNeedsRecreation();
 
 	void PopulateWidgetGeometryCache(FArrangedWidget& Root);
 
@@ -243,7 +243,6 @@ private:
 
 	TSharedPtr<class SZoomPan> PreviewHitTestRoot;
 	TSharedPtr<SBox> PreviewAreaConstraint;
-	TSharedPtr<SBox> PreviewSizeConstraint;
 	TSharedPtr<SDPIScaler> PreviewSurface;
 	TSharedPtr<SCanvas> ExtensionWidgetCanvas;
 	TSharedPtr<SPaintSurface> EffectsLayer;
diff --git a/Engine/Source/Editor/UMGEditor/Private/Details/SPropertyBinding.cpp b/Engine/Source/Editor/UMGEditor/Private/Details/SPropertyBinding.cpp
--- a/Engine/Source/Editor/UMGEditor/Private/Details/SPropertyBinding.cpp
+++ b/Engine/Source/Editor/UMGEditor/Private/Details/SPropertyBinding.cpp
@@ -548,7 +548,7 @@ bool SPropertyBinding::CanRemoveBinding(TSharedRef<IPropertyHandle> PropertyHand
 
 void SPropertyBinding::HandleRemoveBinding(TSharedRef<IPropertyHandle> PropertyHandle)
 {
-	const FScopedTransaction Transaction(LOCTEXT("BindDelegate", "Remove Binding"));
+	const FScopedTransaction Transaction(LOCTEXT("UnbindDelegate", "Remove Binding"));
 
 	Blueprint->Modify();
 
diff --git a/Engine/Source/Editor/UMGEditor/Private/Details/SWidgetDetailsView.cpp b/Engine/Source/Editor/UMGEditor/Private/Details/SWidgetDetailsView.cpp
--- a/Engine/Source/Editor/UMGEditor/Private/Details/SWidgetDetailsView.cpp
+++ b/Engine/Source/Editor/UMGEditor/Private/Details/SWidgetDetailsView.cpp
@@ -500,9 +500,10 @@ void SWidgetDetailsView::NotifyPostChange(const FPropertyChangedEvent& PropertyC
 
 	// If the property that changed is marked as "DesignerRebuild" we invalidate
 	// the preview.
-	if ( PropertyChangedEvent.Property->GetBoolMetaData(DesignerRebuildName) )
+	if ( PropertyChangedEvent.Property->HasMetaData(DesignerRebuildName) || PropertyThatChanged->GetActiveMemberNode()->GetValue()->HasMetaData(DesignerRebuildName) )
 	{
-		BlueprintEditor.Pin()->InvalidatePreview();
+		const bool bViewOnly = true;
+		BlueprintEditor.Pin()->InvalidatePreview(bViewOnly);
 	}
 }
 
diff --git a/Engine/Source/Editor/UMGEditor/Private/Templates/WidgetTemplateBlueprintClass.cpp b/Engine/Source/Editor/UMGEditor/Private/Templates/WidgetTemplateBlueprintClass.cpp
--- a/Engine/Source/Editor/UMGEditor/Private/Templates/WidgetTemplateBlueprintClass.cpp
+++ b/Engine/Source/Editor/UMGEditor/Private/Templates/WidgetTemplateBlueprintClass.cpp
@@ -36,15 +36,14 @@ FText FWidgetTemplateBlueprintClass::GetCategory() const
 	}
 	else
 	{
-		const FString* FoundPaletteCategory = WidgetAssetData.TagsAndValues.Find(GET_MEMBER_NAME_CHECKED(UWidgetBlueprint, PaletteCategory));
-
-		if ( FoundPaletteCategory && !FoundPaletteCategory->IsEmpty() )
+		//If the blueprint is unloaded we need to extract it from the asset metadata.
+		FText FoundPaletteCategoryText = WidgetAssetData.GetTagValueRef<FText>(GET_MEMBER_NAME_CHECKED(UWidgetBlueprint, PaletteCategory));
+		if (!FoundPaletteCategoryText.IsEmpty())
 		{
-			return FText::FromString(*FoundPaletteCategory);
+			return FoundPaletteCategoryText;
 		}
 		else
 		{
-			//If the blueprint is unloaded we need to extract it from the asset metadata.
 			auto DefaultUserWidget = UUserWidget::StaticClass()->GetDefaultObject<UUserWidget>();
 			return DefaultUserWidget->GetPaletteCategory();
 		}
@@ -77,12 +76,12 @@ const FSlateBrush* FWidgetTemplateBlueprintClass::GetIcon() const
 TSharedRef<IToolTip> FWidgetTemplateBlueprintClass::GetToolTip() const
 {
 	FText Description;
-	if ( const FString* DescriptionStr = WidgetAssetData.TagsAndValues.Find( GET_MEMBER_NAME_CHECKED( UBlueprint, BlueprintDescription ) ) )
+
+	FString DescriptionStr = WidgetAssetData.GetTagValueRef<FString>( GET_MEMBER_NAME_CHECKED( UBlueprint, BlueprintDescription ) );
+	if ( !DescriptionStr.IsEmpty() )
 	{
-		if ( !DescriptionStr->IsEmpty() )
-		{
-			Description = FText::FromString( DescriptionStr->Replace( TEXT( "\\n" ), TEXT( "\n" ) ) );
-		}
+		DescriptionStr.ReplaceInline( TEXT( "\\n" ), TEXT( "\n" ) );
+		Description = FText::FromString( MoveTemp(DescriptionStr) );
 	}
 	else
 	{
diff --git a/Engine/Source/Editor/UMGEditor/Private/UMGEditorModule.cpp b/Engine/Source/Editor/UMGEditor/Private/UMGEditorModule.cpp
--- a/Engine/Source/Editor/UMGEditor/Private/UMGEditorModule.cpp
+++ b/Engine/Source/Editor/UMGEditor/Private/UMGEditorModule.cpp
@@ -85,23 +85,11 @@ public:
 				GetMutableDefault<UWidgetDesignerSettings>()
 				);
 		}
-
-		const ISlateStyle* UMGStyle = FSlateStyleRegistry::FindSlateStyle("UMGStyle");
-		if (UMGStyle)
-		{
-			FClassIconFinder::RegisterIconSource(UMGStyle);
-		}
 	}
 
 	/** Called before the module is unloaded, right before the module object is destroyed. */
 	virtual void ShutdownModule() override
 	{
-		const ISlateStyle* UMGStyle = FSlateStyleRegistry::FindSlateStyle("UMGStyle");
-		if (UMGStyle)
-		{
-			FClassIconFinder::UnregisterIconSource(UMGStyle);
-		}
-
 		MenuExtensibilityManager.Reset();
 		ToolBarExtensibilityManager.Reset();
 
diff --git a/Engine/Source/Editor/UMGEditor/Private/WidgetBlueprintEditor.cpp b/Engine/Source/Editor/UMGEditor/Private/WidgetBlueprintEditor.cpp
--- a/Engine/Source/Editor/UMGEditor/Private/WidgetBlueprintEditor.cpp
+++ b/Engine/Source/Editor/UMGEditor/Private/WidgetBlueprintEditor.cpp
@@ -234,9 +234,16 @@ const TSet< TWeakObjectPtr<UObject> >& FWidgetBlueprintEditor::GetSelectedObject
 	return SelectedObjects;
 }
 
-void FWidgetBlueprintEditor::InvalidatePreview()
+void FWidgetBlueprintEditor::InvalidatePreview(bool bViewOnly)
 {
-	bPreviewInvalidated = true;
+	if ( bViewOnly )
+	{
+		OnWidgetPreviewUpdated.Broadcast();
+	}
+	else
+	{
+		bPreviewInvalidated = true;
+	}
 }
 
 void FWidgetBlueprintEditor::OnBlueprintChangedImpl(UBlueprint* InBlueprint, bool bIsJustBeingCompiled )
diff --git a/Engine/Source/Editor/UMGEditor/Private/WidgetBlueprintEditor.h b/Engine/Source/Editor/UMGEditor/Private/WidgetBlueprintEditor.h
--- a/Engine/Source/Editor/UMGEditor/Private/WidgetBlueprintEditor.h
+++ b/Engine/Source/Editor/UMGEditor/Private/WidgetBlueprintEditor.h
@@ -14,7 +14,7 @@ class UUserWidget;
 /**
  * widget blueprint editor (extends Blueprint editor)
  */
-class FWidgetBlueprintEditor : public FBlueprintEditor
+class UMGEDITOR_API FWidgetBlueprintEditor : public FBlueprintEditor
 {
 public:
 	DECLARE_MULTICAST_DELEGATE_OneParam(FOnHoveredWidgetSet, const FWidgetReference&)
@@ -56,7 +56,7 @@ public:
 	void SetIsSimulating(bool bSimulating);
 
 	/** Causes the preview to be destroyed and a new one to be created next tick */
-	void InvalidatePreview();
+	void InvalidatePreview(bool bViewOnly = false);
 
 	/** Immediately rebuilds the preview widget. */
 	void RefreshPreview();
@@ -91,9 +91,6 @@ public:
 	/** @return The selected set of Objects */
 	const TSet< TWeakObjectPtr<UObject> >& GetSelectedObjects() const;
 
-	/** @return Notification for when the preview widget has been updated */
-	FOnWidgetPreviewUpdated& GetOnWidgetPreviewUpdated() { return OnWidgetPreviewUpdated; }
-
 	TSharedPtr<class FWidgetBlueprintEditorToolbar> GetWidgetToolbarBuilder() { return WidgetToolbar; }
 
 	/** Migrate a property change from the preview GUI to the template GUI. */
@@ -137,6 +134,9 @@ public:
 	/** Fires whenever the selected set of widgets changes */
 	FOnSelectedWidgetsChanged OnSelectedWidgetsChanged;
 
+	/** Notification for when the preview widget has been updated */
+	FOnWidgetPreviewUpdated OnWidgetPreviewUpdated;
+
 	/** Command list for handling widget actions in the WidgetBlueprintEditor */
 	TSharedPtr< FUICommandList > DesignerCommandList;
 
@@ -218,9 +218,6 @@ private:
 	/** The preview GUI object */
 	mutable TWeakObjectPtr<UUserWidget> PreviewWidgetPtr;
 
-	/** Notification for when the preview widget has been updated */
-	FOnWidgetPreviewUpdated OnWidgetPreviewUpdated;
-
 	/** Delegate called when a undo/redo transaction happens */
 	FOnWidgetBlueprintTransaction OnWidgetBlueprintTransaction;
 
