Disable the message log while running in unattended mode to prevent the build machine from attemting to open windows while running automation tests

diff --git a/Engine/Source/Editor/MainFrame/Private/MainFrameModule.cpp b/Engine/Source/Editor/MainFrame/Private/MainFrameModule.cpp
--- a/Engine/Source/Editor/MainFrame/Private/MainFrameModule.cpp
+++ b/Engine/Source/Editor/MainFrame/Private/MainFrameModule.cpp
@@ -216,7 +216,7 @@ void FMainFrameModule::CreateDefaultMainFrame( const bool bStartImmersive, const
 
 			// make sure we only allow the message log to be shown when we have a level editor main tab
 			FMessageLogModule& MessageLogModule = FModuleManager::LoadModuleChecked<FMessageLogModule>(TEXT("MessageLog"));
-			MessageLogModule.EnableMessageLogDisplay(true);
+			MessageLogModule.EnableMessageLogDisplay(!FApp::IsUnattended());
 		}
 
 		// Initialize the main frame window
