Fix crash processing mainframe keybindings during slow tasks

diff --git a/Engine/Source/Editor/MainFrame/Private/Frame/MainFrameActions.cpp b/Engine/Source/Editor/MainFrame/Private/Frame/MainFrameActions.cpp
--- a/Engine/Source/Editor/MainFrame/Private/Frame/MainFrameActions.cpp
+++ b/Engine/Source/Editor/MainFrame/Private/Frame/MainFrameActions.cpp
@@ -209,13 +209,16 @@ void FMainFrameCommands::RegisterCommands()
 
 FReply FMainFrameActionCallbacks::OnUnhandledKeyDownEvent(const FKeyEvent& InKeyEvent)
 {
-	if(FMainFrameCommands::ActionList->ProcessCommandBindings(InKeyEvent))
+	if(!GIsSlowTask)
 	{
-		return FReply::Handled();
-	}
-	else if(FPlayWorldCommands::GlobalPlayWorldActions.IsValid() && FPlayWorldCommands::GlobalPlayWorldActions->ProcessCommandBindings(InKeyEvent))
-	{
-		return FReply::Handled();
+		if (FMainFrameCommands::ActionList->ProcessCommandBindings(InKeyEvent))
+		{
+			return FReply::Handled();
+		}
+		else if (FPlayWorldCommands::GlobalPlayWorldActions.IsValid() && FPlayWorldCommands::GlobalPlayWorldActions->ProcessCommandBindings(InKeyEvent))
+		{
+			return FReply::Handled();
+		}
 	}
 
 	return FReply::Unhandled();
