Fix client targets showing up in packaging menu for installed builds (which don't ship with client binaries).

(This implementation is a little suspect in general - there is no requirement that client targets are called "*Client.Target.cs", though that is typically the case.)

diff --git a/Engine/Source/Editor/MainFrame/Private/Menus/PackageProjectMenu.h b/Engine/Source/Editor/MainFrame/Private/Menus/PackageProjectMenu.h
--- a/Engine/Source/Editor/MainFrame/Private/Menus/PackageProjectMenu.h
+++ b/Engine/Source/Editor/MainFrame/Private/Menus/PackageProjectMenu.h
@@ -190,7 +190,10 @@ protected:
 
 		// Add the Client configurations if there is a {ProjectName}Client.Target.cs file.
 		TArray<FString> ClientTargetFileNames;
-		IFileManager::Get().FindFiles(ClientTargetFileNames, *(FPaths::GameSourceDir() / TEXT("*client.target.cs")), true, false);
+		if (FInstalledPlatformInfo::Get().IsValidPlatformType(PlatformInfo::EPlatformType::Client))
+		{
+			IFileManager::Get().FindFiles(ClientTargetFileNames, *(FPaths::GameSourceDir() / TEXT("*client.target.cs")), true, false);
+		}
 
 		if (ClientTargetFileNames.Num() > 0)
 		{
