Do not save window position if running with -nullrhi (UE-52498).

- This also fixes a crash on exiting automation tests.

diff --git a/Engine/Source/Editor/MainFrame/Private/Frame/MainFrameHandler.cpp b/Engine/Source/Editor/MainFrame/Private/Frame/MainFrameHandler.cpp
--- a/Engine/Source/Editor/MainFrame/Private/Frame/MainFrameHandler.cpp
+++ b/Engine/Source/Editor/MainFrame/Private/Frame/MainFrameHandler.cpp
@@ -27,7 +27,7 @@ void FMainFrameHandler::ShutDownEditor()
 	TSharedPtr<SWindow> RootWindow = RootWindowPtr.Pin();
 
 	// Save root window placement so we can restore it.
-	if (RootWindow.IsValid())
+	if (!GUsingNullRHI && RootWindow.IsValid())
 	{
 		FSlateRect WindowRect = RootWindow->GetNonMaximizedRectInScreen();
 
