Added a "report bug" menu entry to the help menu

diff --git a/Engine/Config/BaseEditor.ini b/Engine/Config/BaseEditor.ini
--- a/Engine/Config/BaseEditor.ini
+++ b/Engine/Config/BaseEditor.ini
@@ -314,6 +314,9 @@ EpicGamesURL="http://UnrealEngine.com"
 FacebookURL="http://www.facebook.com/UnrealEngine"
 ViewportControlsURL="http://docs.unrealengine.com/latest/INT/Engine/UI/LevelEditor/Viewports/ViewportControls/index.html"
 MeshSimplificationPluginsURL="/ue/marketplace/content-cat/assets/codeplugins"
+ReportABugURL="https://epicsupport.force.com/unrealengine/s/"
+IssueTrackerURL="https://issues.unrealengine.com/"
+
 
 [UnrealEd.URLOverrides]
 UDNURL="http://docs.unrealengine.com/latest/INT/index.html"
diff --git a/Engine/Source/Editor/LevelEditor/Private/LevelEditorMenu.cpp b/Engine/Source/Editor/LevelEditor/Private/LevelEditorMenu.cpp
--- a/Engine/Source/Editor/LevelEditor/Private/LevelEditorMenu.cpp
+++ b/Engine/Source/Editor/LevelEditor/Private/LevelEditorMenu.cpp
@@ -237,7 +237,7 @@ TSharedRef< SWidget > FLevelEditorMenu::MakeLevelEditorMenu( const TSharedPtr<FU
 
 	// Extend the Help menu
 	Extender->AddMenuExtension(
-		"HelpOnline",
+		"BugReporting",
 		EExtensionHook::Before,
 		CommandList.ToSharedRef(),
 		FMenuExtensionDelegate::CreateStatic( &Local::ExtendHelpMenu ) );
diff --git a/Engine/Source/Editor/MainFrame/Private/Frame/MainFrameActions.cpp b/Engine/Source/Editor/MainFrame/Private/Frame/MainFrameActions.cpp
--- a/Engine/Source/Editor/MainFrame/Private/Frame/MainFrameActions.cpp
+++ b/Engine/Source/Editor/MainFrame/Private/Frame/MainFrameActions.cpp
@@ -170,6 +170,12 @@ void FMainFrameCommands::RegisterCommands()
 	UI_COMMAND(VisitForums, "Forums...", "Go the the Unreal Engine forums to view announcements and engage in discussions with other developers.", EUserInterfaceActionType::Button, FInputChord());
 	ActionList->MapAction(VisitForums, FExecuteAction::CreateStatic(&FMainFrameActionCallbacks::VisitForums));
 
+	UI_COMMAND(ReportABug, "Report a Bug...", "Found a bug?  Go here to fill out a bug report", EUserInterfaceActionType::Button, FInputChord());
+	ActionList->MapAction(ReportABug, FExecuteAction::CreateStatic(&FMainFrameActionCallbacks::ReportABug));
+
+	UI_COMMAND(OpenIssueTracker, "Issue Tracker", "Go here to view the Unreal Engine bug tracking website", EUserInterfaceActionType::Button, FInputChord());
+	ActionList->MapAction(OpenIssueTracker, FExecuteAction::CreateStatic(&FMainFrameActionCallbacks::OpenIssueTracker));
+
 	UI_COMMAND( VisitAskAQuestionPage, "Ask a Question...", "Have a question?  Go here to ask about anything and everything related to Unreal.", EUserInterfaceActionType::Button, FInputChord() );
 	ActionList->MapAction( VisitAskAQuestionPage, FExecuteAction::CreateStatic( &FMainFrameActionCallbacks::VisitAskAQuestionPage ) );
 
@@ -1021,6 +1027,24 @@ bool FMainFrameActionCallbacks::OpenSlateApp_IsChecked( FName AppName )
 	return false;
 }
 
+void FMainFrameActionCallbacks::ReportABug()
+{
+	FString ReportABugURL;
+	if (FUnrealEdMisc::Get().GetURL(TEXT("ReportABugURL"), ReportABugURL, false))
+	{
+		FPlatformProcess::LaunchURL(*ReportABugURL, NULL, NULL);
+	}
+}
+
+void FMainFrameActionCallbacks::OpenIssueTracker()
+{
+	FString IssueTrackerURL;
+	if (FUnrealEdMisc::Get().GetURL(TEXT("IssueTrackerURL"), IssueTrackerURL, false))
+	{
+		FPlatformProcess::LaunchURL(*IssueTrackerURL, NULL, NULL);
+	}
+}
+
 void FMainFrameActionCallbacks::VisitAskAQuestionPage()
 {
 	FString AskAQuestionURL;
diff --git a/Engine/Source/Editor/MainFrame/Private/Frame/MainFrameActions.h b/Engine/Source/Editor/MainFrame/Private/Frame/MainFrameActions.h
--- a/Engine/Source/Editor/MainFrame/Private/Frame/MainFrameActions.h
+++ b/Engine/Source/Editor/MainFrame/Private/Frame/MainFrameActions.h
@@ -54,6 +54,8 @@ public:
 	TSharedPtr< FUICommandInfo > VisitForums;
 	TSharedPtr< FUICommandInfo > VisitAskAQuestionPage;
 	TSharedPtr< FUICommandInfo > VisitSearchForAnswersPage;
+	TSharedPtr< FUICommandInfo > ReportABug;
+	TSharedPtr< FUICommandInfo > OpenIssueTracker;
 	TSharedPtr< FUICommandInfo > VisitSupportWebSite;
 	TSharedPtr< FUICommandInfo > VisitEpicGamesDotCom;
 	TSharedPtr< FUICommandInfo > AboutUnrealEd;
@@ -131,6 +133,12 @@ public:
 	/** Checks if a slate window is already open */
 	static bool OpenSlateApp_IsChecked( FName AppName );
 
+	/** Visits the report a bug page */
+	static void ReportABug();
+
+	/** Opens the issue tracker page */
+	static void OpenIssueTracker();
+
 	/** Visits the "ask a question" page on UDN */
 	static void VisitAskAQuestionPage();
 
diff --git a/Engine/Source/Editor/MainFrame/Private/Menus/MainMenu.cpp b/Engine/Source/Editor/MainFrame/Private/Menus/MainMenu.cpp
--- a/Engine/Source/Editor/MainFrame/Private/Menus/MainMenu.cpp
+++ b/Engine/Source/Editor/MainFrame/Private/Menus/MainMenu.cpp
@@ -284,7 +284,14 @@ void FMainMenu::FillWindowMenu( FMenuBuilder& MenuBuilder, const TSharedRef< FEx
 
 void FMainMenu::FillHelpMenu( FMenuBuilder& MenuBuilder, const TSharedRef< FExtender > Extender )
 {
-	MenuBuilder.BeginSection("HelpOnline", NSLOCTEXT("MainHelpMenu", "Online", "Online"));
+	MenuBuilder.BeginSection("BugReporting", NSLOCTEXT("MainHelpMenu", "BugsReporting", "Bugs"));
+	{
+		MenuBuilder.AddMenuEntry(FMainFrameCommands::Get().ReportABug);
+		MenuBuilder.AddMenuEntry(FMainFrameCommands::Get().OpenIssueTracker);
+	}
+	MenuBuilder.EndSection();
+
+	MenuBuilder.BeginSection("HelpOnline", NSLOCTEXT("MainHelpMenu", "Online", "Help Online"));
 	{
 		MenuBuilder.AddMenuEntry(FMainFrameCommands::Get().VisitSupportWebSite);
 		MenuBuilder.AddMenuEntry(FMainFrameCommands::Get().VisitForums);
