Fixes for running an editor with a unique build environment.

* Add support for running UAT with an absolute path to the editor executable, and pass a full path when cooking or packaging.
* Fix path to the editor executable when launching a standalone editor instance for preview.

diff --git a/Engine/Source/Editor/MainFrame/Private/Frame/MainFrameActions.cpp b/Engine/Source/Editor/MainFrame/Private/Frame/MainFrameActions.cpp
--- a/Engine/Source/Editor/MainFrame/Private/Frame/MainFrameActions.cpp
+++ b/Engine/Source/Editor/MainFrame/Private/Frame/MainFrameActions.cpp
@@ -412,7 +412,7 @@ void FMainFrameActionCallbacks::CookContent(const FName InPlatformInfoName)
 	}
 
 	FString ProjectPath = FPaths::IsProjectFilePathSet() ? FPaths::ConvertRelativePathToFull(FPaths::GetProjectFilePath()) : FPaths::RootDir() / FApp::GetProjectName() / FApp::GetProjectName() + TEXT(".uproject");
-	FString CommandLine = FString::Printf(TEXT("-ScriptsForProject=\"%s\" BuildCookRun %s%s -nop4 -project=\"%s\" -cook -skipstage -ue4exe=%s %s -utf8output"),
+	FString CommandLine = FString::Printf(TEXT("-ScriptsForProject=\"%s\" BuildCookRun %s%s -nop4 -project=\"%s\" -cook -skipstage -ue4exe=\"%s\" %s -utf8output"),
 		*ProjectPath,
 		GetUATCompilationFlags(),
 		FApp::IsEngineInstalled() ? TEXT(" -installed") : TEXT(""),
@@ -742,7 +742,7 @@ void FMainFrameActionCallbacks::PackageProject( const FName InPlatformInfoName )
 	}
 
 	FString ProjectPath = FPaths::IsProjectFilePathSet() ? FPaths::ConvertRelativePathToFull(FPaths::GetProjectFilePath()) : FPaths::RootDir() / FApp::GetProjectName() / FApp::GetProjectName() + TEXT(".uproject");
-	FString CommandLine = FString::Printf(TEXT("-ScriptsForProject=\"%s\" BuildCookRun %s%s -nop4 -project=\"%s\" -cook -stage -archive -archivedirectory=\"%s\" -package -clientconfig=%s -ue4exe=%s %s -utf8output"),
+	FString CommandLine = FString::Printf(TEXT("-ScriptsForProject=\"%s\" BuildCookRun %s%s -nop4 -project=\"%s\" -cook -stage -archive -archivedirectory=\"%s\" -package -clientconfig=%s -ue4exe=\"%s\" %s -utf8output"),
 		*ProjectPath,
 		GetUATCompilationFlags(),
 		FApp::IsEngineInstalled() ? TEXT(" -installed") : TEXT(""),
diff --git a/Engine/Source/Editor/UnrealEd/Private/PlayLevel.cpp b/Engine/Source/Editor/UnrealEd/Private/PlayLevel.cpp
--- a/Engine/Source/Editor/UnrealEd/Private/PlayLevel.cpp
+++ b/Engine/Source/Editor/UnrealEd/Private/PlayLevel.cpp
@@ -1470,7 +1470,7 @@ void UEditorEngine::PlayStandaloneLocalPc(FString MapNameOverride, FIntPoint* Wi
 	}
 
 	// launch the game process
-	FString GamePath = FPlatformProcess::GenerateApplicationPath(FApp::GetName(), FApp::GetBuildConfiguration());
+	FString GamePath = FString(FPlatformProcess::BaseDir()) / FString(FPlatformProcess::ExecutableName(false));
 	FPlayOnPCInfo *NewSession = new (PlayOnLocalPCSessions) FPlayOnPCInfo();
 
 	uint32 ProcessID = 0;
diff --git a/Engine/Source/Editor/UnrealEd/Private/UnrealEdMisc.cpp b/Engine/Source/Editor/UnrealEd/Private/UnrealEdMisc.cpp
--- a/Engine/Source/Editor/UnrealEd/Private/UnrealEdMisc.cpp
+++ b/Engine/Source/Editor/UnrealEd/Private/UnrealEdMisc.cpp
@@ -1790,7 +1790,7 @@ bool FUnrealEdMisc::GetURL( const TCHAR* InKey, FString& OutURL, const bool bChe
 
 FString FUnrealEdMisc::GetExecutableForCommandlets() const
 {
-	FString ExecutableName = FPlatformProcess::ExecutableName(false);
+	FString ExecutableName = FString(FPlatformProcess::BaseDir()) / FString(FPlatformProcess::ExecutableName(false));
 #if PLATFORM_WINDOWS
 	// turn UE4editor into UE4editor-cmd
 	if(ExecutableName.EndsWith(".exe", ESearchCase::IgnoreCase) && !FPaths::GetBaseFilename(ExecutableName).EndsWith("-cmd", ESearchCase::IgnoreCase))
diff --git a/Engine/Source/Programs/AutomationTool/AutomationUtils/LinuxHostPlatform.cs b/Engine/Source/Programs/AutomationTool/AutomationUtils/LinuxHostPlatform.cs
--- a/Engine/Source/Programs/AutomationTool/AutomationUtils/LinuxHostPlatform.cs
+++ b/Engine/Source/Programs/AutomationTool/AutomationUtils/LinuxHostPlatform.cs
@@ -24,6 +24,11 @@ namespace AutomationTool
 
 		public override string GetUE4ExePath(string UE4Exe)
 		{
+			if(Path.IsPathRooted(UE4Exe))
+			{
+				return CommandUtils.CombinePaths(UE4Exe);
+			}
+
 			int CmdExeIndex = UE4Exe.IndexOf("-Cmd.exe");
 			if (CmdExeIndex != -1)
 			{
diff --git a/Engine/Source/Programs/AutomationTool/AutomationUtils/MacHostPlatform.cs b/Engine/Source/Programs/AutomationTool/AutomationUtils/MacHostPlatform.cs
--- a/Engine/Source/Programs/AutomationTool/AutomationUtils/MacHostPlatform.cs
+++ b/Engine/Source/Programs/AutomationTool/AutomationUtils/MacHostPlatform.cs
@@ -24,6 +24,11 @@ namespace AutomationTool
 
 		public override string GetUE4ExePath(string UE4Exe)
 		{
+			if(Path.IsPathRooted(UE4Exe))
+			{
+				return CommandUtils.CombinePaths(UE4Exe);
+			}
+
 			int CmdExeIndex = UE4Exe.IndexOf("-Cmd.exe");
 			if (CmdExeIndex != -1)
 			{
diff --git a/Engine/Source/Programs/AutomationTool/AutomationUtils/WindowsHostPlatform.cs b/Engine/Source/Programs/AutomationTool/AutomationUtils/WindowsHostPlatform.cs
--- a/Engine/Source/Programs/AutomationTool/AutomationUtils/WindowsHostPlatform.cs
+++ b/Engine/Source/Programs/AutomationTool/AutomationUtils/WindowsHostPlatform.cs
@@ -1,4 +1,4 @@
-﻿// Copyright 1998-2019 Epic Games, Inc. All Rights Reserved.
+// Copyright 1998-2019 Epic Games, Inc. All Rights Reserved.
 
 using System;
 using System.Collections.Generic;
@@ -24,7 +24,14 @@ namespace AutomationTool
 
 		public override string GetUE4ExePath(string UE4Exe)
 		{
-			return CommandUtils.CombinePaths(CommandUtils.CmdEnv.LocalRoot, RelativeBinariesFolder, UE4Exe);
+			if(Path.IsPathRooted(UE4Exe))
+			{
+				return CommandUtils.CombinePaths(UE4Exe);
+			}
+			else
+			{
+				return CommandUtils.CombinePaths(CommandUtils.CmdEnv.LocalRoot, RelativeBinariesFolder, UE4Exe);
+			}
 		}
 
 		public override string LocalBuildsLogFolder
