Make Device Output Log non-experimental

diff --git a/Engine/Source/Developer/OutputLog/Private/OutputLogModule.cpp b/Engine/Source/Developer/OutputLog/Private/OutputLogModule.cpp
--- a/Engine/Source/Developer/OutputLog/Private/OutputLogModule.cpp
+++ b/Engine/Source/Developer/OutputLog/Private/OutputLogModule.cpp
@@ -99,8 +99,7 @@ void FOutputLogModule::StartupModule()
 		.SetDisplayName(NSLOCTEXT("UnrealEditor", "DeviceOutputLogTab", "Device Output Log"))
 		.SetTooltipText(NSLOCTEXT("UnrealEditor", "DeviceOutputLogTooltipText", "Open the Device Output Log tab."))
 		.SetGroup( WorkspaceMenu::GetMenuStructure().GetDeveloperToolsLogCategory() )
-		.SetIcon( FSlateIcon(FEditorStyle::GetStyleSetName(), "Log.TabIcon") )
-		.SetAutoGenerateMenuEntry(false); // remove once not Experimental
+		.SetIcon( FSlateIcon(FEditorStyle::GetStyleSetName(), "Log.TabIcon") );
 	
 	OutputLogHistory = MakeShareable(new FOutputLogHistory);
 }
diff --git a/Engine/Source/Editor/MainFrame/Private/Menus/MainMenu.cpp b/Engine/Source/Editor/MainFrame/Private/Menus/MainMenu.cpp
--- a/Engine/Source/Editor/MainFrame/Private/Menus/MainMenu.cpp
+++ b/Engine/Source/Editor/MainFrame/Private/Menus/MainMenu.cpp
@@ -212,10 +212,9 @@ void FMainMenu::FillWindowMenu( FMenuBuilder& MenuBuilder, const TSharedRef< FEx
 		bool bBlutility = GetDefault<UEditorExperimentalSettings>()->bEnableEditorUtilityBlueprints;
 		bool bLocalizationDashboard = GetDefault<UEditorExperimentalSettings>()->bEnableLocalizationDashboard;
 		bool bTranslationPicker = GetDefault<UEditorExperimentalSettings>()->bEnableTranslationPicker;
-		bool bDeviceOutputLog = GetDefault<UEditorExperimentalSettings>()->bDeviceOutputLog;
 
 		// Make sure at least one is enabled before creating the section
-		if (bBlutility || bLocalizationDashboard || bTranslationPicker || bDeviceOutputLog)
+		if (bBlutility || bLocalizationDashboard || bTranslationPicker)
 		{
 			MenuBuilder.BeginSection("ExperimentalTabSpawners", LOCTEXT("ExperimentalTabSpawnersHeading", "Experimental"));
 			{
@@ -251,18 +250,6 @@ void FMainMenu::FillWindowMenu( FMenuBuilder& MenuBuilder, const TSharedRef< FEx
 						FUIAction(FExecuteAction::CreateStatic(&FMainFrameTranslationEditorMenu::HandleOpenTranslationPicker))
 						);
 				}
-				
-				// Device output log
-				if (bDeviceOutputLog)
-				{
-					MenuBuilder.AddMenuEntry(
-						LOCTEXT("DeviceOutputLogMenuLabel", "Device Output Log"),
-						LOCTEXT("DeviceOutputLogToolTip", "Open the Device Output Log tab."),
-						FSlateIcon(FEditorStyle::GetStyleSetName(), "Log.TabIcon"),
-						FUIAction(FExecuteAction::CreateStatic(&FMainMenu::OpenDeviceOutputLog))
-						);
-				}
-
 			}
 			MenuBuilder.EndSection();
 		}
diff --git a/Engine/Source/Editor/MainFrame/Private/Menus/MainMenu.h b/Engine/Source/Editor/MainFrame/Private/Menus/MainMenu.h
--- a/Engine/Source/Editor/MainFrame/Private/Menus/MainMenu.h
+++ b/Engine/Source/Editor/MainFrame/Private/Menus/MainMenu.h
@@ -111,13 +111,4 @@ private:
 		FModuleManager::Get().LoadModuleChecked<IModuleInterface>("LogVisualizer");
 		FGlobalTabmanager::Get()->InvokeTab(FName(TEXT("VisualLogger")));
 	}
-		
-	/**
-	* Opens the 'Device Output Log' tab.
-	* Remove this when it is no longer experimental.
-	*/
-	static void OpenDeviceOutputLog()
-	{
-		FGlobalTabmanager::Get()->InvokeTab(FName(TEXT("DeviceOutputLog")));
-	}
 };
diff --git a/Engine/Source/Editor/UnrealEd/Classes/Settings/EditorExperimentalSettings.h b/Engine/Source/Editor/UnrealEd/Classes/Settings/EditorExperimentalSettings.h
--- a/Engine/Source/Editor/UnrealEd/Classes/Settings/EditorExperimentalSettings.h
+++ b/Engine/Source/Editor/UnrealEd/Classes/Settings/EditorExperimentalSettings.h
@@ -46,11 +46,6 @@ public:
 	UPROPERTY(EditAnywhere, config, Category = Tools, meta = (DisplayName = "Enable Details Panel Favorites"))
 	bool bEnableFavoriteSystem;
 
-
-	/** Device output log window (currently implemented for Android, Switch & XboxOne only)*/
-	UPROPERTY(EditAnywhere, config, Category = Tools, meta = (DisplayName = "Device Output Log"))
-	bool bDeviceOutputLog;
-
 	/** Specify which console-specific nomenclature to use for gamepad label text */
 	UPROPERTY(EditAnywhere, config, Category=UserInterface, meta=(DisplayName="Console for Gamepad Labels"))
 	TEnumAsByte<EConsoleForGamepadLabels::Type> ConsoleForGamepadLabels;
