Editor fullscreen disabled for Linux, missed a check here

diff --git a/Engine/Source/Editor/MainFrame/Private/Menus/MainMenu.cpp b/Engine/Source/Editor/MainFrame/Private/Menus/MainMenu.cpp
--- a/Engine/Source/Editor/MainFrame/Private/Menus/MainMenu.cpp
+++ b/Engine/Source/Editor/MainFrame/Private/Menus/MainMenu.cpp
@@ -247,7 +247,7 @@ void FMainMenu::FillWindowMenu( FMenuBuilder& MenuBuilder, const TSharedRef< FEx
 	{
 		MenuBuilder.AddMenuEntry(FMainFrameCommands::Get().ResetLayout);
 		MenuBuilder.AddMenuEntry(FMainFrameCommands::Get().SaveLayout);
-#if !PLATFORM_MAC // On Mac windowed fullscreen mode in the editor is currently unavailable
+#if !PLATFORM_MAC || !PLATFORM_LIUNX // On Mac/Linux windowed fullscreen mode in the editor is currently unavailable
 		MenuBuilder.AddMenuEntry(FMainFrameCommands::Get().ToggleFullscreen);
 #endif
 	}
