PR #5392: UE-61357: Show message prompt when packaging with missing game default map (Contributed by [REDACTED_PERSON])

diff --git a/Engine/Source/Editor/MainFrame/Private/Frame/MainFrameActions.cpp b/Engine/Source/Editor/MainFrame/Private/Frame/MainFrameActions.cpp
--- a/Engine/Source/Editor/MainFrame/Private/Frame/MainFrameActions.cpp
+++ b/Engine/Source/Editor/MainFrame/Private/Frame/MainFrameActions.cpp
@@ -51,6 +51,7 @@
 #include "Settings/EditorSettings.h"
 #include "AnalyticsEventAttribute.h"
 #include "Kismet2/DebuggerCommands.h"
+#include "GameMapsSettings.h"
 
 #define LOCTEXT_NAMESPACE "MainFrameActions"
 
@@ -462,6 +463,12 @@ void FMainFrameActionCallbacks::PackageProject( const FName InPlatformInfoName )
 		return;
 	}
 
+	if (UGameMapsSettings::GetGameDefaultMap().IsEmpty())
+	{
+		FMessageDialog::Open(EAppMsgType::Ok, LOCTEXT("MissingGameDefaultMap", "No Game Default Map specified in Project Settings > Maps & Modes."));
+		return;
+	}
+
 	if (PlatformInfo->SDKStatus == PlatformInfo::EPlatformSDKStatus::NotInstalled || (bProjectHasCode && PlatformInfo->bUsesHostCompiler && !FSourceCodeNavigation::IsCompilerAvailable()))
 	{
 		IMainFrameModule& MainFrameModule = FModuleManager::GetModuleChecked<IMainFrameModule>(TEXT("MainFrame"));
