Disabled Window->Enable Fullscreen menu item on Mac as it doesn't work properly on Mac (freezes the editor). We can reenable this once we make the editor compatible with Mac full screen behavior

diff --git a/Engine/Source/Editor/MainFrame/Private/Frame/MainFrameActions.cpp b/Engine/Source/Editor/MainFrame/Private/Frame/MainFrameActions.cpp
--- a/Engine/Source/Editor/MainFrame/Private/Frame/MainFrameActions.cpp
+++ b/Engine/Source/Editor/MainFrame/Private/Frame/MainFrameActions.cpp
@@ -201,12 +201,14 @@ void FMainFrameCommands::RegisterCommands()
 	UI_COMMAND( SaveLayout, "Save Layout", "Save the layout customizations", EUserInterfaceActionType::Button, FInputChord() );
 	ActionList->MapAction( SaveLayout, FExecuteAction::CreateStatic( &FMainFrameActionCallbacks::SaveLayout) );
 
+#if !PLATFORM_MAC // Fullscreen mode in the editor is currently unsupported on Mac
 	UI_COMMAND( ToggleFullscreen, "Enable Fullscreen", "Enables fullscreen mode for the application, expanding across the entire monitor", EUserInterfaceActionType::ToggleButton, FInputChord(EModifierKey::Shift, EKeys::F11) );
 	ActionList->MapAction( ToggleFullscreen,
 		FExecuteAction::CreateStatic( &FMainFrameActionCallbacks::ToggleFullscreen_Execute ),
 		FCanExecuteAction(),
 		FIsActionChecked::CreateStatic( &FMainFrameActionCallbacks::FullScreen_IsChecked )
 	);
+#endif
 
 	UI_COMMAND(OpenWidgetReflector, "Open Widget Reflector", "Opens the Widget Reflector", EUserInterfaceActionType::Button, FInputChord(EModifierKey::Shift | EModifierKey::Control , EKeys::W));
 	ActionList->MapAction(OpenWidgetReflector, FExecuteAction::CreateStatic(&FMainFrameActionCallbacks::OpenWidgetReflector_Execute));
@@ -939,6 +941,7 @@ void FMainFrameActionCallbacks::SaveLayout()
 
 void FMainFrameActionCallbacks::ToggleFullscreen_Execute()
 {
+#if !PLATFORM_MAC // Fullscreen mode in the editor is currently unsupported on Mac
 	if ( GIsEditor && FApp::HasProjectName() )
 	{
 		static TWeakPtr<SDockTab> LevelEditorTabPtr = FGlobalTabmanager::Get()->InvokeTab(FTabId("LevelEditor"));
@@ -952,7 +955,8 @@ void FMainFrameActionCallbacks::ToggleFullscreen_Execute()
 		{
 			LevelEditorWindow->SetWindowMode(EWindowMode::Windowed);
 		}
-	}	
+	}
+#endif
 }
 
 bool FMainFrameActionCallbacks::FullScreen_IsChecked()
