Deprecate the IDesktopPlatformBase::GetSolutionPath() method, and change existing code to use the SourceCodeAccess module instead.

diff --git a/Engine/Source/Developer/DesktopPlatform/Private/DesktopPlatformBase.h b/Engine/Source/Developer/DesktopPlatform/Private/DesktopPlatformBase.h
--- a/Engine/Source/Developer/DesktopPlatform/Private/DesktopPlatformBase.h
+++ b/Engine/Source/Developer/DesktopPlatform/Private/DesktopPlatformBase.h
@@ -10,6 +10,8 @@
 class FEngineVersion;
 class FJsonObject;
 
+PRAGMA_DISABLE_DEPRECATION_WARNINGS // For FDesktopPlatformBase::GetSolutionPath() implementation
+
 class FDesktopPlatformBase : public IDesktopPlatform
 {
 public:
@@ -83,3 +85,5 @@ protected:
 	FString GetUnrealBuildToolProjectFileName(const FString& RootDir) const;
 	FString GetUnrealBuildToolExecutableFilename(const FString& RootDir) const;	
 };
+
+PRAGMA_ENABLE_DEPRECATION_WARNINGS
diff --git a/Engine/Source/Developer/DesktopPlatform/Public/IDesktopPlatform.h b/Engine/Source/Developer/DesktopPlatform/Public/IDesktopPlatform.h
--- a/Engine/Source/Developer/DesktopPlatform/Public/IDesktopPlatform.h
+++ b/Engine/Source/Developer/DesktopPlatform/Public/IDesktopPlatform.h
@@ -419,6 +419,7 @@ public:
 	* @param OutSolutionPath	Receives the string 
 	* @return True if a solution file exists and OutSolutionPath has been updated
 	*/
+	UE_DEPRECATED(4.24, "The IDesktopPlatform::GetSolutionPath() method has been deprecated. Use the ISourceCodeAccessor interface for cross-IDE support instead.")
 	virtual bool GetSolutionPath(FString& OutSolutionPath) = 0;
 
 	/**
diff --git a/Engine/Source/Editor/Documentation/Private/UDNParser.cpp b/Engine/Source/Editor/Documentation/Private/UDNParser.cpp
--- a/Engine/Source/Editor/Documentation/Private/UDNParser.cpp
+++ b/Engine/Source/Editor/Documentation/Private/UDNParser.cpp
@@ -1319,15 +1319,7 @@ bool FUDNParser::ParseCodeLink(FString &InternalLink)
 	}
 
 	// Finally create the complete path - project name and all
-	int32 PathEndIndex;
-	FString SolutionPath;
-	if( FDesktopPlatformModule::Get()->GetSolutionPath(SolutionPath) && SolutionPath.FindLastChar(TEXT('/'), PathEndIndex) == true)
-	{
-		SolutionPath = SolutionPath.LeftChop(SolutionPath.Len() - PathEndIndex - 1);
-		SolutionPath += Path;
-		bLinkParsedOK = SourceCodeAccessor.OpenFileAtLine(SolutionPath, Line, Col);
-	}
-	return bLinkParsedOK;
+	return SourceCodeAccessor.OpenFileAtLine(FPaths::RootDir() + Path, Line, Col);
 }
 
 bool FUDNParser::ParseAssetLink(FString &InternalLink)
diff --git a/Engine/Source/Editor/Kismet/Kismet.Build.cs b/Engine/Source/Editor/Kismet/Kismet.Build.cs
--- a/Engine/Source/Editor/Kismet/Kismet.Build.cs
+++ b/Engine/Source/Editor/Kismet/Kismet.Build.cs
@@ -19,6 +19,7 @@ public class Kismet : ModuleRules
                 "LevelEditor",
 				"GameProjectGeneration",
 				"Sequencer",
+				"SourceCodeAccess",
 			}
 			);
 
diff --git a/Engine/Source/Editor/Kismet/Private/SBlueprintEditorToolbar.cpp b/Engine/Source/Editor/Kismet/Private/SBlueprintEditorToolbar.cpp
--- a/Engine/Source/Editor/Kismet/Private/SBlueprintEditorToolbar.cpp
+++ b/Engine/Source/Editor/Kismet/Private/SBlueprintEditorToolbar.cpp
@@ -28,6 +28,8 @@
 #include "SBlueprintEditorSelectedDebugObjectWidget.h"
 #include "DesktopPlatformModule.h"
 #include "SBlueprintRevisionMenu.h"
+#include "ISourceCodeAccessor.h"
+#include "ISourceCodeAccessModule.h"
 
 #define LOCTEXT_NAMESPACE "KismetToolbar"
 
@@ -79,8 +81,8 @@ void FKismet2Menu::FillFileMenuBlueprintSection( FMenuBuilder& MenuBuilder, FBlu
 	MenuBuilder.EndSection();
 
 	// Only show the developer menu on machines with the solution (assuming they can build it)
-	FString SolutionPath;
-	if(FDesktopPlatformModule::Get()->GetSolutionPath(SolutionPath))
+	ISourceCodeAccessModule* SourceCodeAccessModule = FModuleManager::GetModulePtr<ISourceCodeAccessModule>("SourceCodeAccess");
+	if(SourceCodeAccessModule != nullptr && SourceCodeAccessModule->GetAccessor().CanAccessSourceCode())
 	{
 		MenuBuilder.BeginSection("FileDeveloper");
 		{
diff --git a/Engine/Source/Editor/MainFrame/Private/Frame/MainFrameActions.cpp b/Engine/Source/Editor/MainFrame/Private/Frame/MainFrameActions.cpp
--- a/Engine/Source/Editor/MainFrame/Private/Frame/MainFrameActions.cpp
+++ b/Engine/Source/Editor/MainFrame/Private/Frame/MainFrameActions.cpp
@@ -900,17 +900,7 @@ void FMainFrameActionCallbacks::OpenIDE()
 	{
 		if ( !FSourceCodeNavigation::OpenModuleSolution() )
 		{
-			FString SolutionPath;
-			if(FDesktopPlatformModule::Get()->GetSolutionPath(SolutionPath))
-			{
-				const FString FullPath = IFileManager::Get().ConvertToAbsolutePathForExternalAppForRead( *SolutionPath );
-				const FText Message = FText::Format( LOCTEXT( "OpenIDEFailed_MissingFile", "Could not open {0} for project {1}" ), FSourceCodeNavigation::GetSelectedSourceCodeIDE(), FText::FromString( FullPath ) );
-				FMessageDialog::Open( EAppMsgType::Ok, Message );
-			}
-			else
-			{
-				FMessageDialog::Open( EAppMsgType::Ok, LOCTEXT("OpenIDEFailed_MissingSolution", "Couldn't find solution"));
-			}
+			FMessageDialog::Open( EAppMsgType::Ok, LOCTEXT("OpenIDEFailed_UnableToOpenSolution", "Unable to open solution"));
 		}
 	}
 }
diff --git a/Engine/Source/Runtime/Toolbox/Private/ToolboxModule.cpp b/Engine/Source/Runtime/Toolbox/Private/ToolboxModule.cpp
--- a/Engine/Source/Runtime/Toolbox/Private/ToolboxModule.cpp
+++ b/Engine/Source/Runtime/Toolbox/Private/ToolboxModule.cpp
@@ -20,6 +20,8 @@
 #include "Widgets/Testing/STestSuite.h"
 #include "Editor/WorkspaceMenuStructure/Public/WorkspaceMenuStructure.h"
 #include "Editor/WorkspaceMenuStructure/Public/WorkspaceMenuStructureModule.h"
+#include "ISourceCodeAccessor.h"
+#include "ISourceCodeAccessModule.h"
 
 #if !UE_BUILD_SHIPPING
 #include "ISlateReflectorModule.h"
@@ -30,8 +32,8 @@ static bool bTabsRegistered = false;
 
 static bool CanShowModulesTab()
 {
-	FString SolutionPath;
-	return FDesktopPlatformModule::Get()->GetSolutionPath(SolutionPath);
+	ISourceCodeAccessModule* SourceCodeAccessModule = FModuleManager::GetModulePtr<ISourceCodeAccessModule>("SourceCodeAccess");
+	return SourceCodeAccessModule != nullptr && SourceCodeAccessModule->GetAccessor().CanAccessSourceCode();
 }
 
 class SDebugPanel : public SCompoundWidget
diff --git a/Engine/Source/Runtime/Toolbox/Toolbox.Build.cs b/Engine/Source/Runtime/Toolbox/Toolbox.Build.cs
--- a/Engine/Source/Runtime/Toolbox/Toolbox.Build.cs
+++ b/Engine/Source/Runtime/Toolbox/Toolbox.Build.cs
@@ -10,7 +10,8 @@ public class Toolbox : ModuleRules
 			new string[] {
 				"GammaUI",
 				"MainFrame",
-				"ModuleUI"
+				"ModuleUI",
+				"SourceCodeAccess"
 			}
 		);
 
