Remove empty "Build Targets" menu item for content projects.

#rnx

diff --git a/Engine/Source/Editor/MainFrame/Private/Menus/PackageProjectMenu.h b/Engine/Source/Editor/MainFrame/Private/Menus/PackageProjectMenu.h
--- a/Engine/Source/Editor/MainFrame/Private/Menus/PackageProjectMenu.h
+++ b/Engine/Source/Editor/MainFrame/Private/Menus/PackageProjectMenu.h
@@ -90,11 +90,16 @@ public:
 			LOCTEXT("PackageProjectBuildConfigurationSubMenuToolTip", "Select the build configuration to package the project with"),
 			FNewMenuDelegate::CreateStatic(&FPackageProjectMenu::MakeBuildConfigurationsMenu)
 		);
-		MenuBuilder.AddSubMenu(
-			LOCTEXT("PackageProjectBuildTargetSubMenuLabel", "Build Target"),
-			LOCTEXT("PackageProjectBuildTargetSubMenuToolTip", "Select the build target to package"),
-			FNewMenuDelegate::CreateStatic(&FPackageProjectMenu::MakeBuildTargetsMenu)
-		);
+
+		IDesktopPlatform* DesktopPlatform = FDesktopPlatformModule::Get();
+		if (DesktopPlatform->GetTargetsForCurrentProject().Num() > 0)
+		{
+			MenuBuilder.AddSubMenu(
+				LOCTEXT("PackageProjectBuildTargetSubMenuLabel", "Build Target"),
+				LOCTEXT("PackageProjectBuildTargetSubMenuToolTip", "Select the build target to package"),
+				FNewMenuDelegate::CreateStatic(&FPackageProjectMenu::MakeBuildTargetsMenu)
+			);
+		}
 
 		MenuBuilder.AddMenuSeparator();
 		MenuBuilder.AddMenuEntry(FMainFrameCommands::Get().PackagingSettings);
