Copied CL 9371593
Modified Remove All User-Created Layouts message to fix grammar errors and follow LOCTEXT format.

diff --git a/Engine/Source/Editor/MainFrame/Private/Menus/LayoutsMenu.cpp b/Engine/Source/Editor/MainFrame/Private/Menus/LayoutsMenu.cpp
--- a/Engine/Source/Editor/MainFrame/Private/Menus/LayoutsMenu.cpp
+++ b/Engine/Source/Editor/MainFrame/Private/Menus/LayoutsMenu.cpp
@@ -855,8 +855,7 @@ void FLayoutsMenuRemove::RemoveUserLayouts()
 	if (NumberUserLayoutFiles > 0)
 	{
 		// Are you sure you want to do this?
-		const FText TextNumberUserLayoutFiles  = FText::FromString(FString::FromInt(NumberUserLayoutFiles));
-		const FText TextBody = FText::Format(LOCTEXT("ActionRemoveAllUserLayoutMsg", "Are you sure you want to permanently remove all the {0} layout profiles created by the user? This action cannot be undone."), TextNumberUserLayoutFiles);
+		const FText TextBody = FText::Format(LOCTEXT("ActionRemoveAllUserLayoutMsg", "Are you sure you want to permanently remove {0} layout {0}|plural(one=profile,other=profiles)? This action cannot be undone."), NumberUserLayoutFiles);
 		const FText TextTitle = LOCTEXT("RemoveAllUserLayouts_Title", "Remove All User-Created Layouts");
 		if (EAppReturnType::Ok != OpenMsgDlgInt(EAppMsgType::OkCancel, TextBody, TextTitle))
 		{
