Copied CL# 9541296
Made layout name and description dialog more clear by replacing the Yes/No/Cancel answers by [REDACTED_PERSON] / Clear Values / Cancel.
Also allow SCustomDialog to be used outside the UnrealEd module (by using UNREALED_API) and fixed some compiling errors on its commented example.

diff --git a/Engine/Source/Editor/MainFrame/Private/Menus/LayoutsMenu.cpp b/Engine/Source/Editor/MainFrame/Private/Menus/LayoutsMenu.cpp
--- a/Engine/Source/Editor/MainFrame/Private/Menus/LayoutsMenu.cpp
+++ b/Engine/Source/Editor/MainFrame/Private/Menus/LayoutsMenu.cpp
@@ -18,12 +18,15 @@
 #include "Misc/ConfigCacheIni.h"
 #include "Misc/FileHelper.h"
 #include "Widgets/Notifications/SNotificationList.h"
+#include "Widgets/SBoxPanel.h"
 // Developer
 #include "IDesktopPlatform.h"
 #include "DesktopPlatformModule.h"
 #include "ToolMenus.h"
 // Editor
 #include "Classes/EditorStyleSettings.h"
+#include "Dialogs/CustomDialog.h"
+#include "Misc/MessageDialog.h"
 #include "Dialogs/Dialogs.h"
 #include "Editor/EditorPerProjectUserSettings.h"
 #include "Frame/MainFrameActions.h"
@@ -98,31 +101,56 @@ bool TrySaveLayoutOrWarnInternal(const FString& InSourceFilePath, const FString&
 				if (bShouldAskBeforeCleaningLayoutNameAndDescriptionFields)
 				{
 					// Open Dialog
-					const FText TextTitle = LOCTEXT("OverrideLayoutNameAndDescriptionFieldBodyTitle", "Clean UI Layout name and description fields");
+					const FText TextTitle = LOCTEXT("OverrideLayoutNameAndDescriptionFieldBodyTitle", "Preserve UI Layout Name and Description Fields?");
 					const FText TextBody = FText::Format(
 						LOCTEXT("OverrideLayoutNameAndDescriptionFieldBody",
 							"You are saving a layout that contains a custom layout name and/or description. Do you also want to copy these 2 properties?\n"
 							" - Current layout name: {0}\n"
 							" - Current layout description: {1}\n\n"
-							"If you select \"Yes\", the displayed name and description of the original layout customization will also be copied into the new configuration file.\n\n"
-							"If you select \"No\", these fields will be emptied.\n\n"
-							"If you are not sure, select \"Yes\" if you are exporting the layout configuration without making any changes, or \"No\" if you have made or plan to make changes to the layout.\n\n"
+							"If you select \"Preserve Values\", the displayed name and description of the original layout customization will also be copied into the new configuration file.\n\n"
+							"If you select \"Clear Values\", these fields will be emptied.\n\n"
+							"If you are not sure, select \"Preserve Values\" if you are exporting the layout configuration without making any changes, or \"Clear Values\" if you have made or plan to make changes to the layout.\n\n"
 						),
 						LayoutNameSource, LayoutDescriptionSource);
-					const int AppReturnType = OpenMsgDlgInt(EAppMsgType::YesNoCancel, TextBody, TextTitle);
-					// Handle user answers
-					if (AppReturnType == EAppReturnType::Yes)
+					// Dialog SWidget
+					TSharedRef<SVerticalBox> DialogContents = SNew(SVerticalBox);
+					DialogContents->AddSlot()
+						.Padding(0, 16, 0, 0)
+						[
+							SNew(STextBlock)
+							.Text(TextBody)
+						];
+					const FText PreserveValuesText = LOCTEXT("PreserveValuesText", "Preserve Values");
+					const FText ClearValuesText = LOCTEXT("ClearValuesText", "Clear Values");
+					const FText CancelText = NSLOCTEXT("Dialogs", "EAppReturnTypeCancel", "Cancel");
+					TSharedRef<SCustomDialog> CustomDialog = SNew(SCustomDialog)
+						.Title(TextTitle)
+						.DialogContent(DialogContents)
+						.Buttons({ SCustomDialog::FButton(PreserveValuesText), SCustomDialog::FButton(ClearValuesText), SCustomDialog::FButton(CancelText) })
+					;
+					// Returns 0 when "Preserve Values" is pressed, 1 when "Clear Values" is pressed, or 2 when Cancel/Esc is pressed
+					const int ButtonPressed = CustomDialog->ShowModal();
+					// Preserve Values
+					if (ButtonPressed == 0)
 					{
 						bCleanLayoutNameAndDescriptionFields = false;
 					}
-					else if (AppReturnType == EAppReturnType::No)
+					// Clear Values
+					else if (ButtonPressed == 1)
 					{
 						bCleanLayoutNameAndDescriptionFields = true;
 					}
-					else if (AppReturnType == EAppReturnType::Cancel)
+					// Cancel or Esc or window closed
+					else if (ButtonPressed == 2 || ButtonPressed == -1)
 					{
 						return false;
 					}
+					// This should never occur
+					else
+					{
+						ensureMsgf(false, TEXT("This option should never occur, something went wrong!"));
+						return false;
+					}
 				}
 			}
 		}
diff --git a/Engine/Source/Editor/UnrealEd/Public/Dialogs/CustomDialog.h b/Engine/Source/Editor/UnrealEd/Public/Dialogs/CustomDialog.h
--- a/Engine/Source/Editor/UnrealEd/Public/Dialogs/CustomDialog.h
+++ b/Engine/Source/Editor/UnrealEd/Public/Dialogs/CustomDialog.h
@@ -9,15 +9,17 @@
  * It also supports adding a custom icon to the dialog.
  * Usage:
  * TSharedRef<SCustomDialog> HelloWorldDialog = SNew(SCustomDialog)
-		.Title(FText("Hello, World!"))
-		.DialogContent( SNew(SImage).Image(FName("Hello")) )
-		.Buttons( { SCustomDialog::FButton(LOCTEXT("OK", "OK"), []() { printf("OK pressed!" ); } ),
-				SCustomDialog::FButton(LOCTEXT("Cancel", "Cancel") } );
-
-		// returns 0 when OK is pressed, 1 when Cancel is pressed
-		int ButtonPressed = CustomDialog->ShowModal(); 
+		.Title(FText(LOCTEXT("HelloWorldTitleExample", "Hello, World!")))
+		.DialogContent( SNew(SImage).Image(FName(TEXT("Hello"))))
+		.Buttons({
+			SCustomDialog::FButton(LOCTEXT("OK", "OK")),
+			SCustomDialog::FButton(LOCTEXT("Cancel", "Cancel"))
+		});
+
+   // returns 0 when OK is pressed, 1 when Cancel is pressed, -1 if the window is closed
+   const int ButtonPressed = HelloWorldDialog->ShowModal();
  */
-class SCustomDialog : public SWindow
+class UNREALED_API SCustomDialog : public SWindow
 {
 public:
 	struct FButton
