Fix debug and test configurations being displayed in the project packaging menu.

#rnx

diff --git a/Engine/Source/Editor/MainFrame/Private/Menus/PackageProjectMenu.h b/Engine/Source/Editor/MainFrame/Private/Menus/PackageProjectMenu.h
--- a/Engine/Source/Editor/MainFrame/Private/Menus/PackageProjectMenu.h
+++ b/Engine/Source/Editor/MainFrame/Private/Menus/PackageProjectMenu.h
@@ -213,22 +213,27 @@ protected:
 	 */
 	static void MakeBuildConfigurationsMenu(FMenuBuilder& MenuBuilder)
 	{
+		EProjectType ProjectType = FGameProjectGenerationModule::Get().ProjectHasCodeFiles() ? EProjectType::Code : EProjectType::Content;
+
 		TArray<EProjectPackagingBuildConfigurations> PackagingConfigurations = UProjectPackagingSettings::GetValidPackageConfigurations();
 		for(EProjectPackagingBuildConfigurations PackagingConfiguration : PackagingConfigurations)
 		{
 			const UProjectPackagingSettings::FConfigurationInfo& Info = UProjectPackagingSettings::ConfigurationInfo[(int)PackagingConfiguration];
-			MenuBuilder.AddMenuEntry(
-				Info.Name, 
-				Info.ToolTip,
-				FSlateIcon(),
-				FUIAction(
-					FExecuteAction::CreateStatic(&FMainFrameActionCallbacks::PackageBuildConfiguration, PackagingConfiguration),
-					FCanExecuteAction::CreateStatic(&FMainFrameActionCallbacks::CanPackageBuildConfiguration, PackagingConfiguration),
-					FIsActionChecked::CreateStatic(&FMainFrameActionCallbacks::PackageBuildConfigurationIsChecked, PackagingConfiguration)
-				),
-				NAME_None,
-				EUserInterfaceActionType::RadioButton
-			);
+			if (FInstalledPlatformInfo::Get().IsValid(TOptional<EBuildTargetType>(), TOptional<FString>(), Info.Configuration, ProjectType, EInstalledPlatformState::Downloaded))
+			{
+				MenuBuilder.AddMenuEntry(
+					Info.Name, 
+					Info.ToolTip,
+					FSlateIcon(),
+					FUIAction(
+						FExecuteAction::CreateStatic(&FMainFrameActionCallbacks::PackageBuildConfiguration, PackagingConfiguration),
+						FCanExecuteAction::CreateStatic(&FMainFrameActionCallbacks::CanPackageBuildConfiguration, PackagingConfiguration),
+						FIsActionChecked::CreateStatic(&FMainFrameActionCallbacks::PackageBuildConfigurationIsChecked, PackagingConfiguration)
+					),
+					NAME_None,
+					EUserInterfaceActionType::RadioButton
+				);
+			}
 		}
 	}
 
