Fix the "Open IDE" command not being shown after adding code to a blueprint project.

#rnx

diff --git a/Engine/Source/Editor/MainFrame/Private/Frame/MainFrameActions.cpp b/Engine/Source/Editor/MainFrame/Private/Frame/MainFrameActions.cpp
--- a/Engine/Source/Editor/MainFrame/Private/Frame/MainFrameActions.cpp
+++ b/Engine/Source/Editor/MainFrame/Private/Frame/MainFrameActions.cpp
@@ -111,7 +111,7 @@ void FMainFrameCommands::RegisterCommands()
 	ActionList->MapAction( RefreshCodeProject, FExecuteAction::CreateStatic( &FMainFrameActionCallbacks::RefreshCodeProject ), FCanExecuteAction::CreateStatic( &FMainFrameActionCallbacks::IsCodeProject ) );
 
 	UI_COMMAND( OpenIDE, "Open IDE", "Opens your C++ code in an integrated development environment.", EUserInterfaceActionType::Button, FInputChord() );
-	ActionList->MapAction( OpenIDE, FExecuteAction::CreateStatic( &FMainFrameActionCallbacks::OpenIDE ), FCanExecuteAction::CreateStatic( &FMainFrameActionCallbacks::IsCodeProject ) );
+	ActionList->MapAction( OpenIDE, FExecuteAction::CreateStatic( &FMainFrameActionCallbacks::OpenIDE ), FCanExecuteAction::CreateStatic( &FMainFrameActionCallbacks::IsCodeProject ), FGetActionCheckState(), FIsActionButtonVisible::CreateStatic( &FMainFrameActionCallbacks::CanOpenIDE ) );
 
 	UI_COMMAND( ZipUpProject, "Zip Up Project", "Zips up the project into a zip file.", EUserInterfaceActionType::Button, FInputChord() );
 	ActionList->MapAction(ZipUpProject, FExecuteAction::CreateStatic( &FMainFrameActionCallbacks::ZipUpProject ), DefaultExecuteAction);
@@ -902,6 +902,11 @@ void FMainFrameActionCallbacks::OpenIDE()
 	}
 }
 
+bool FMainFrameActionCallbacks::CanOpenIDE()
+{
+	return FSourceCodeNavigation::DoesModuleSolutionExist();
+}
+
 void FMainFrameActionCallbacks::ZipUpProject()
 {
 #if PLATFORM_WINDOWS
diff --git a/Engine/Source/Editor/MainFrame/Private/Frame/MainFrameActions.h b/Engine/Source/Editor/MainFrame/Private/Frame/MainFrameActions.h
--- a/Engine/Source/Editor/MainFrame/Private/Frame/MainFrameActions.h
+++ b/Engine/Source/Editor/MainFrame/Private/Frame/MainFrameActions.h
@@ -227,6 +227,9 @@ public:
 	/** Opens an IDE to edit c++ code */
 	static void OpenIDE();
 
+	/** Determines whether we can open the IDE to edit c++ code */
+	static bool CanOpenIDE();
+
 	/** Zips up the project */
 	static void ZipUpProject();
 
diff --git a/Engine/Source/Editor/MainFrame/Private/Menus/MainMenu.cpp b/Engine/Source/Editor/MainFrame/Private/Menus/MainMenu.cpp
--- a/Engine/Source/Editor/MainFrame/Private/Menus/MainMenu.cpp
+++ b/Engine/Source/Editor/MainFrame/Private/Menus/MainMenu.cpp
@@ -460,26 +460,28 @@ void FMainMenu::RegisterFileProjectMenu()
 	);
 	*/
 
-	FString SolutionPath;
-	if (FSourceCodeNavigation::DoesModuleSolutionExist())
+	Section.AddDynamicEntry("CodeProject", FNewToolMenuSectionDelegate::CreateLambda([ShortIDEName](FToolMenuSection& InSection)
 	{
-		Section.AddMenuEntry( FMainFrameCommands::Get().RefreshCodeProject,
-			FText::Format(LOCTEXT("RefreshCodeProjectLabel", "Refresh {0} Project"), ShortIDEName),
-			FText::Format(LOCTEXT("RefreshCodeProjectTooltip", "Refreshes your C++ code project in {0}."), ShortIDEName)
-		);
+		if (FSourceCodeNavigation::DoesModuleSolutionExist())
+		{
+			InSection.AddMenuEntry( FMainFrameCommands::Get().RefreshCodeProject,
+				FText::Format(LOCTEXT("RefreshCodeProjectLabel", "Refresh {0} Project"), ShortIDEName),
+				FText::Format(LOCTEXT("RefreshCodeProjectTooltip", "Refreshes your C++ code project in {0}."), ShortIDEName)
+			);
+		}
+		else
+		{
+			InSection.AddMenuEntry( FMainFrameCommands::Get().RefreshCodeProject,
+				FText::Format(LOCTEXT("GenerateCodeProjectLabel", "Generate {0} Project"), ShortIDEName),
+				FText::Format(LOCTEXT("GenerateCodeProjectTooltip", "Generates your C++ code project in {0}."), ShortIDEName)
+			);
+		}
+	}));
 
-		Section.AddMenuEntry( FMainFrameCommands::Get().OpenIDE,
-			FText::Format(LOCTEXT("OpenIDELabel", "Open {0}"), ShortIDEName),
-			FText::Format(LOCTEXT("OpenIDETooltip", "Opens your C++ code in {0}."), ShortIDEName)
-		);
-	}
-	else
-	{
-		Section.AddMenuEntry( FMainFrameCommands::Get().RefreshCodeProject,
-			FText::Format(LOCTEXT("GenerateCodeProjectLabel", "Generate {0} Project"), ShortIDEName),
-			FText::Format(LOCTEXT("GenerateCodeProjectTooltip", "Generates your C++ code project in {0}."), ShortIDEName)
-		);
-	}
+	Section.AddMenuEntry( FMainFrameCommands::Get().OpenIDE,
+		FText::Format(LOCTEXT("OpenIDELabel", "Open {0}"), ShortIDEName),
+		FText::Format(LOCTEXT("OpenIDETooltip", "Opens your C++ code in {0}."), ShortIDEName)
+	);
 
 	Section.AddDynamicEntry("CookContentForPlatform", FNewToolMenuSectionDelegate::CreateLambda([](FToolMenuSection& InSection)
 	{
