Copied CL# 10088684

New features (they are very related to each other, so they have been integrated together):
1) UE-80569: When the user attempts to save a custom layout, we open a Slate SWidget save dialog that no longer allows them to enter the file path where the layout will be saved, but instead the Layout Name and (optionally) Description. The file name will automatically be generated from the displayed layout name.
2) UE-80575: When attempting to save a custom layout, we prevent the OS save dialog. Since the user cannot actually save a layout anywhere other than the default save location, this streamlines the workflow and minimize confusion. Instead, the Slate SWidget detailed in point 1 is used, where the user can introduce only the displayed layout name, and the file name and its directory are automatically generated (fixed directory).
3) UE-80577: After 1 and 2, it is unclear where the layout is saved. Thus, when saving a custom layout, we provide an Editor toast notifying the user of their custom layout's save location. It is analog to the one provided when a user takes a screenshot within the Editor.
In addition, the possible messages displayed to the user when the layout file was not successfully copied have also been improved (higher details and fixed some grammar typos).

Some general and additional considerations about the new Slate layout save dialog:
1) The "Save" button will be disabled when the introduced file name is empty or an error occurs (e.g., name is too long).  If any error occurs, the dialog will display the error message as well, analog to "File" --> "Save current as...".
2) An Override/Cancel dialog will be displayed if the user selects a file name that matches another existing one.
3) An automatic NSLOCTEXT key and namespace fields are generated for each profile created. This is abstracted from the user, who does not need to know about this. Translations come later from a LocRes file that we generate for the engine. Thus, only the profile files that we have included in P4 might be translated.

Some potential issues / observations:
1) The export layout remains with the old OS save dialog (as it can be saved in different directories). This makes the appearance of both of them very different, but it also makes export similar to import.

diff --git a/Engine/Source/Editor/MainFrame/Private/Menus/LayoutsMenu.cpp b/Engine/Source/Editor/MainFrame/Private/Menus/LayoutsMenu.cpp
--- a/Engine/Source/Editor/MainFrame/Private/Menus/LayoutsMenu.cpp
+++ b/Engine/Source/Editor/MainFrame/Private/Menus/LayoutsMenu.cpp
@@ -26,11 +26,11 @@
 // Editor
 #include "Classes/EditorStyleSettings.h"
 #include "Dialogs/CustomDialog.h"
-#include "Misc/MessageDialog.h"
-#include "Dialogs/Dialogs.h"
 #include "Editor/EditorPerProjectUserSettings.h"
 #include "Frame/MainFrameActions.h"
 #include "LevelViewportActions.h"
+#include "Menus/SaveLayoutDialog.h"
+#include "Misc/MessageDialog.h"
 #include "UnrealEdGlobals.h"
 #include "UnrealEdMisc.h"
 
@@ -76,7 +76,9 @@ TArray<FString> GetIniFilesInFolderInternal(const FString& InStringDirectory)
 	return LayoutIniFileNames;
 }
 
-bool TrySaveLayoutOrWarnInternal(const FString& InSourceFilePath, const FString& InTargetFilePath, const FText& InWhatIsThis, const bool bCleanLayoutNameAndDescriptionFieldsIfNoSameValues, const bool bShouldAskBeforeCleaningLayoutNameAndDescriptionFields = false)
+bool TrySaveLayoutOrWarnInternal(
+	const FString& InSourceFilePath, const FString& InTargetFilePath, const FText& InWhatIsThis, const bool bCleanLayoutNameAndDescriptionFieldsIfNoSameValues,
+	const bool bShouldAskBeforeCleaningLayoutNameAndDescriptionFields = false, const bool bShowSaveToast = false)
 {
 	// If desired, ask user whether to keep the LayoutName and LayoutDescription fields
 	bool bCleanLayoutNameAndDescriptionFields = false;
@@ -156,6 +158,7 @@ bool TrySaveLayoutOrWarnInternal(const FString& InSourceFilePath, const FString&
 		}
 	}
 	// Copy: Replace main layout with desired one
+	const FString TargetAbsoluteFilePath = FPaths::ConvertRelativePathToFull(InTargetFilePath);
 	const bool bShouldReplace = true;
 	const bool bCopyEvenIfReadOnly = true;
 	const bool bCopyAttributes = false; // If true, it could e.g., copy the read-only flag of DefaultLayout.ini and make all the save/load stuff stop working
@@ -169,37 +172,37 @@ bool TrySaveLayoutOrWarnInternal(const FString& InSourceFilePath, const FString&
 		if (!FPaths::FileExists(InSourceFilePath))
 		{
 			Arguments.Add(TEXT("FileName"), FText::FromString(FPaths::ConvertRelativePathToFull(InSourceFilePath)));
-			TextBody = FText::Format(LOCTEXT("UnsuccessfulSave_NoExist_Notification", "Unsuccessful {WhatIs}, the desired file does not exist. File path:\n{FileName}"), Arguments);
+			TextBody = FText::Format(LOCTEXT("UnsuccessfulSave_NoExist_Notification", "The requested operation ({WhatIs}) was unsuccessful, the desired file does not exist. File path:\n{FileName}"), Arguments);
 			EditorErrors.Warning(TextBody);
 		}
 		// Target is read-only
 		else if (IFileManager::Get().IsReadOnly(*InTargetFilePath))
 		{
-			Arguments.Add(TEXT("FileName"), FText::FromString(FPaths::ConvertRelativePathToFull(InTargetFilePath)));
-			TextBody = FText::Format(LOCTEXT("UnsuccessfulSave_ReadOnly_Notification", "Unsuccessful {WhatIs}, the target file path is read-only. File path:\n{FileName}"), Arguments);
+			Arguments.Add(TEXT("FileName"), FText::FromString(TargetAbsoluteFilePath));
+			TextBody = FText::Format(LOCTEXT("UnsuccessfulSave_ReadOnly_Notification", "The requested operation ({WhatIs}) was unsuccessful, the target file path is read-only. File path:\n{FileName}"), Arguments);
 			EditorErrors.Warning(TextBody);
 		}
 		// Target and source are the same
-		else if (FPaths::ConvertRelativePathToFull(InTargetFilePath) == FPaths::ConvertRelativePathToFull(InSourceFilePath))
+		else if (TargetAbsoluteFilePath == FPaths::ConvertRelativePathToFull(InSourceFilePath))
 		{
 			Arguments.Add(TEXT("SourceFileName"), FText::FromString(FPaths::ConvertRelativePathToFull(InSourceFilePath)));
-			Arguments.Add(TEXT("FinalFileName"), FText::FromString(FPaths::ConvertRelativePathToFull(InTargetFilePath)));
+			Arguments.Add(TEXT("FinalFileName"), FText::FromString(TargetAbsoluteFilePath));
 			TextBody = FText::Format(LOCTEXT("UnsuccessfulSave_Fallback_Notification",
-				"Unsuccessful {WhatIs}, target and source layout file paths are the same ({SourceFileName})!\nAre you trying to import or replace a file that is already in the layouts folder? If so, remove the current file first."), Arguments);
+				"The requested operation ({WhatIs}) was unsuccessful, target and source layout file paths are the same ({SourceFileName})!\nAre you trying to import or replace a file that is already in the layouts folder? If so, remove the current file first."), Arguments);
 			EditorErrors.Warning(TextBody);
 		}
 		// We don't specifically know why it failed, this is a fallback
 		else
 		{
 			Arguments.Add(TEXT("SourceFileName"), FText::FromString(FPaths::ConvertRelativePathToFull(InSourceFilePath)));
-			Arguments.Add(TEXT("FinalFileName"), FText::FromString(FPaths::ConvertRelativePathToFull(InTargetFilePath)));
-			TextBody = FText::Format(LOCTEXT("UnsuccessfulSave_Fallback_Notification", "Unsuccessful {WhatIs} of:\n{SourceFileName}\nto\n{FinalFileName}"), Arguments);
+			Arguments.Add(TEXT("FinalFileName"), FText::FromString(TargetAbsoluteFilePath));
+			TextBody = FText::Format(LOCTEXT("UnsuccessfulSave_Fallback_Notification", "The requested operation ({WhatIs}) was unsuccessful while copying the layout file from\n{SourceFileName}\ninto\n{FinalFileName}\n\nUsually, this occurs when the introduced file name contains unsupported characters or the total path length exceeds the OS limit."), Arguments);
 			EditorErrors.Warning(TextBody);
 		}
 		EditorErrors.Notify(LOCTEXT("LoadUnsuccessful_Title", "Load Unsuccessful!"));
 		// Show reason
 		const FText TextTitle = LOCTEXT("UnsuccessfulCopyHeader", "Unsuccessful copy!");
-		OpenMsgDlgInt(EAppMsgType::Ok, TextBody, TextTitle);
+		FMessageDialog::Open(EAppMsgType::Ok, TextBody, &TextTitle);
 		// Return
 		return false;
 	}
@@ -219,13 +222,34 @@ bool TrySaveLayoutOrWarnInternal(const FString& InSourceFilePath, const FString&
 			const bool bRead = true;
 			GConfig->Flush(bRead, GEditorLayoutIni);
 			// Re-copy file
-			if (FPaths::ConvertRelativePathToFull(InTargetFilePath) != FPaths::ConvertRelativePathToFull(GEditorLayoutIni))
+			if (TargetAbsoluteFilePath != FPaths::ConvertRelativePathToFull(GEditorLayoutIni))
 			{
 				IFileManager::Get().Copy(*InTargetFilePath, *GEditorLayoutIni, bShouldReplace, bCopyEvenIfReadOnly, bCopyAttributes);
 			}
 		}
 		// Unload target file so it can be re-read into cache properly the next time it is used
 		GConfig->UnloadFile(InTargetFilePath); // We must re-read it to avoid the Editor to use a previously cached name and description
+		// Display Editor toast to inform the user of the result of the operation
+		if (bShowSaveToast)
+		{
+			// Code copied from EditorViewportClient.cpp --> FEditorViewportClient::TakeScreenshot(...) to maintain same format than when saving a screenshot
+			FNotificationInfo Info(FText::GetEmpty());
+			Info.ExpireDuration = 5.0f;
+			Info.bUseSuccessFailIcons = false;
+			Info.bUseLargeFont = false;
+			TSharedPtr<SNotificationItem> SaveMessagePtr = FSlateNotificationManager::Get().AddNotification(Info);
+			if (SaveMessagePtr.IsValid())
+			{
+				const FString& HyperLinkString = TargetAbsoluteFilePath;
+				auto OpenScreenshotFolder = [HyperLinkString]
+				{
+					FPlatformProcess::ExploreFolder(*FPaths::GetPath(HyperLinkString));
+				};
+				SaveMessagePtr->SetText(LOCTEXT("SuccessfulSave_Toast", "Editor layout file saved as"));
+				SaveMessagePtr->SetHyperlink(FSimpleDelegate::CreateLambda(OpenScreenshotFolder), FText::FromString(HyperLinkString));
+				SaveMessagePtr->SetCompletionState(SNotificationItem::CS_Success);
+			}
+		}
 		// Return successful copy message
 		return true;
 	}
@@ -280,13 +304,13 @@ void DisplayLayoutsInternal(FToolMenuSection& InSection, const TArray<TSharedPtr
 // ones that FTabManager::SpawnTab cannot recognized and keeps closed but still saves them in the layout).
 const FString& GetOriginalEditorLayoutIniFilePathInternal()
 {
-	static const FString OriginalEditorLayoutIniFilePath = GEditorLayoutIni + FString("_orig.ini");
+	static const FString OriginalEditorLayoutIniFilePath = GEditorLayoutIni + TEXT("_orig.ini");
 	return OriginalEditorLayoutIniFilePath;
 }
 
 const FString& GetDuplicatedEditorLayoutIniFilePathInternal()
 {
-	static const FString DuplicatedEditorLayoutIniFilePath = GEditorLayoutIni + FString("_temp.ini");
+	static const FString DuplicatedEditorLayoutIniFilePath = GEditorLayoutIni + TEXT("_temp.ini");
 	return DuplicatedEditorLayoutIniFilePath;
 }
 
@@ -527,7 +551,7 @@ void FLayoutsMenuLoad::LoadLayout(const FString& InLayoutPath)
 	const FString& TargetFilePath = GEditorLayoutIni;
 	const bool bCleanLayoutNameAndDescriptionFieldsIfNoSameValues = false;
 	const bool bShouldAskBeforeCleaningLayoutNameAndDescriptionFields = false;
-	const bool SucessfullySaved = TrySaveLayoutOrWarnInternal(SourceFilePath, TargetFilePath, LOCTEXT("LoadLayoutText", "layout load"), bCleanLayoutNameAndDescriptionFieldsIfNoSameValues, bShouldAskBeforeCleaningLayoutNameAndDescriptionFields);
+	const bool SucessfullySaved = TrySaveLayoutOrWarnInternal(SourceFilePath, TargetFilePath, LOCTEXT("LoadLayoutText", "loading the layout"), bCleanLayoutNameAndDescriptionFieldsIfNoSameValues, bShouldAskBeforeCleaningLayoutNameAndDescriptionFields);
 	// Reload current layout
 	if (SucessfullySaved)
 	{
@@ -560,7 +584,7 @@ void FLayoutsMenuLoad::ImportLayout()
 		const FString DefaultFile = "";
 		const bool bWereFilesSelected = DesktopPlatform->OpenFileDialog(
 			FSlateApplication::Get().FindBestParentWindowHandleForDialogs(nullptr),
-			FString("Import a Layout Configuration File"),
+			TEXT("Import a Layout Configuration File"),
 			DefaultDirectory,
 			DefaultFile,
 			TEXT("Layout configuration files|*.ini|"),
@@ -574,7 +598,7 @@ void FLayoutsMenuLoad::ImportLayout()
 			const FString UserLayoutsDirectory = CreateAndGetUserLayoutDirInternal();
 			// Iterate over selected layout ini files
 			FString FirstGoodLayoutFile = "";
-			const FText TrySaveLayoutOrWarnInternalText = LOCTEXT("ImportLayoutText", "layout import");
+			const FText TrySaveLayoutOrWarnInternalText = LOCTEXT("ImportLayoutText", "importing the layout");
 			for (const FString& LayoutFilePath : LayoutFilePaths)
 			{
 				// If file is a layout file, import it
@@ -599,7 +623,7 @@ void FLayoutsMenuLoad::ImportLayout()
 					Arguments.Add(TEXT("FileName"), FText::FromString(FPaths::ConvertRelativePathToFull(LayoutFilePath)));
 					const FText TextBody = FText::Format(LOCTEXT("UnsuccessfulImportBody", "Unsuccessful import, {FileName} is not a layout configuration file!"), Arguments);
 					const FText TextTitle = LOCTEXT("UnsuccessfulImportHeader", "Unsuccessful Import!");
-					OpenMsgDlgInt(EAppMsgType::Ok, TextBody, TextTitle);
+					FMessageDialog::Open(EAppMsgType::Ok, TextBody, &TextTitle);
 				}
 			}
 			// If PIE running, do not reload current layout
@@ -608,7 +632,7 @@ void FLayoutsMenuLoad::ImportLayout()
 				const FText TextBody = LOCTEXT("SuccessfulImportBody",
 					"The layout(s) were successfully imported into the \"User Layouts\" section. However, no layout has been loaded into your current Unreal Editor UI because PIE is currently running. In order to do so, you must stop PIE and then load the layout from the \"User Layouts\" section.");
 				const FText TextTitle = LOCTEXT("SuccessfulImportHeader", "Successful Import!");
-				OpenMsgDlgInt(EAppMsgType::Ok, TextBody, TextTitle);
+				FMessageDialog::Open(EAppMsgType::Ok, TextBody, &TextTitle);
 				return;
 			}
 			// Replace current layout with first one
@@ -631,27 +655,94 @@ void FLayoutsMenuLoad::ImportLayout()
 
 
 
+FText GenerateLocalizedTextForFile(const FText& InText)
+{
+	// Proper FText to FString
+	FString StringSimulatingText;
+	FTextStringHelper::WriteToBuffer(StringSimulatingText, InText);
+	// Sanitize text (truncate if too big)
+	FString SanitazedTruncatedText = (StringSimulatingText.Len() < 100 ? StringSimulatingText : StringSimulatingText.Left(100));
+	FSaveLayoutDialogUtils::SanitizeText(SanitazedTruncatedText);
+	// Create full file path
+	if (StringSimulatingText.Len() < 10 || StringSimulatingText.Left(9) != TEXT("NSLOCTEXT"))
+	{
+		FString StringSimulatingTextRecreated =
+			// Namespace
+			TEXT("NSLOCTEXT(\"LayoutNamespace\", "
+			// Key
+			"\"") + SanitazedTruncatedText + TEXT("\", "
+			// Source string
+			"\"") + StringSimulatingText + TEXT("\")");
+		return FText::FromString(StringSimulatingTextRecreated);
+	}
+	else
+	{
+		return FText::FromString(StringSimulatingText);
+	}
+}
+
 void SaveExportLayoutCommon(const FString& InDefaultDirectory, const bool bMustBeSavedInDefaultDirectory, const FText& InWhatIsThis, const bool bShouldAskBeforeCleaningLayoutNameAndDescriptionFields)
 {
 	// Export/SaveAs the user-selected layout configuration files and load one of them
 	IDesktopPlatform* DesktopPlatform = FDesktopPlatformModule::Get();
 	if (DesktopPlatform)
 	{
-		// Open the "save file" dialog so user can save his/her layout configuration file
+		bool bWereFilesSelected = false;
 		TArray<FString> LayoutFilePaths;
-		const FString DialogTitle = "";
-		const FString DefaultFile = "";
-		const bool bWereFileSelected = DesktopPlatform->SaveFileDialog(
-			FSlateApplication::Get().FindBestParentWindowHandleForDialogs(nullptr),
-			FString("Export a Layout Configuration File"),
-			InDefaultDirectory,
-			DefaultFile,
-			TEXT("Layout configuration files|*.ini|"),
-			EFileDialogFlags::None, //EFileDialogFlags::Multiple, // Allow/Avoid multiple file selection
-			LayoutFilePaths
-		);
+		TArray<FText> LayoutNames;
+		LayoutNames.Emplace(FLayoutSaveRestore::LoadSectionFromConfig(GEditorLayoutIni, "LayoutName"));
+		if (LayoutNames[0].ToString().Len() > 0)
+		{
+			LayoutNames[0] = FText::FromString(TEXT("Copy of ") + LayoutNames[0].ToString());
+		}
+		TArray<FText> LayoutDescriptions;
+		LayoutDescriptions.Emplace(FLayoutSaveRestore::LoadSectionFromConfig(GEditorLayoutIni, "LayoutDescription"));
+		if (LayoutDescriptions[0].ToString().Len() > 0)
+		{
+			LayoutDescriptions[0] = FText::FromString(TEXT("Copy of ") + LayoutDescriptions[0].ToString());
+		}
+		// "Save Layout As..."
+		bool bWasDialogOpened = bMustBeSavedInDefaultDirectory;
+		if (bMustBeSavedInDefaultDirectory)
+		{
+			// Create SWidget for saving the layout in its own SWindow and block the thread until it is finished
+			const TSharedRef<FSaveLayoutDialogParams> SaveLayoutDialogParams = MakeShared<FSaveLayoutDialogParams>(InDefaultDirectory, TEXT(".ini"), LayoutNames, LayoutDescriptions);
+			bWasDialogOpened = FSaveLayoutDialogUtils::CreateSaveLayoutAsDialogInStandaloneWindow(SaveLayoutDialogParams);
+			bWereFilesSelected = SaveLayoutDialogParams->bWereFilesSelected;
+			LayoutFilePaths = SaveLayoutDialogParams->LayoutFilePaths;
+			LayoutNames = SaveLayoutDialogParams->LayoutNames;
+			LayoutDescriptions = SaveLayoutDialogParams->LayoutDescriptions;
+
+			// Update GEditorLayoutIni file if LayoutNames or LayoutDescriptions were modified by the user
+			if (bWasDialogOpened && LayoutNames.Num() > 0 && LayoutDescriptions.Num() > 0 && (LayoutNames[0].ToString().Len() > 0 || LayoutDescriptions[0].ToString().Len() > 0))
+			{
+				const FText& LayoutNameAsTextText = GenerateLocalizedTextForFile(LayoutNames[0]);
+				const FText& LayoutDescriptionAsTextText = GenerateLocalizedTextForFile(LayoutDescriptions[0]);
+				// Update fields
+				FLayoutSaveRestore::SaveSectionToConfig(GEditorLayoutIni, "LayoutName", LayoutNameAsTextText);
+				FLayoutSaveRestore::SaveSectionToConfig(GEditorLayoutIni, "LayoutDescription", LayoutDescriptions[0]);
+				// Flush file
+				const bool bRead = true;
+				GConfig->Flush(bRead, GEditorLayoutIni);
+			}
+		}
+		// "Export Layout..." (or "Save Layout As..." dialog could not be opened)
+		if (!bWasDialogOpened)
+		{
+			// Open the "save file" dialog so user can save his/her layout configuration file
+			const FString DefaultFile = "";
+			bWereFilesSelected = DesktopPlatform->SaveFileDialog(
+				FSlateApplication::Get().FindBestParentWindowHandleForDialogs(nullptr),
+				TEXT("Export a Layout Configuration File"),
+				InDefaultDirectory,
+				DefaultFile,
+				TEXT("Layout configuration files|*.ini|"),
+				EFileDialogFlags::None, //EFileDialogFlags::Multiple, // Allow/Avoid multiple file selection
+				LayoutFilePaths
+			);
+		}
 		// If file(s) selected, copy them into the user layouts directory and load one of them
-		if (bWereFileSelected && LayoutFilePaths.Num() > 0)
+		if (bWereFilesSelected && LayoutFilePaths.Num() > 0)
 		{
 			// Iterate over selected layout ini files
 			FString FirstGoodLayoutFile = "";
@@ -665,14 +756,16 @@ void SaveExportLayoutCommon(const FString& InDefaultDirectory, const bool bMustB
 					// Save in the user layout folder
 					const FString& SourceFilePath = GEditorLayoutIni;
 					const FString& TargetFilePath = LayoutFilePath;
-					const bool bCleanLayoutNameAndDescriptionFieldsIfNoSameValues = true;
-					TrySaveLayoutOrWarnInternal(SourceFilePath, TargetFilePath, InWhatIsThis, bCleanLayoutNameAndDescriptionFieldsIfNoSameValues, bShouldAskBeforeCleaningLayoutNameAndDescriptionFields);
+					const bool bCleanLayoutNameAndDescriptionFieldsIfNoSameValues = !bMustBeSavedInDefaultDirectory;
+					const bool bShowSaveToast = true;
+					TrySaveLayoutOrWarnInternal(SourceFilePath, TargetFilePath, InWhatIsThis, bCleanLayoutNameAndDescriptionFieldsIfNoSameValues, bShouldAskBeforeCleaningLayoutNameAndDescriptionFields, bShowSaveToast);
 				}
 				// If trying to write in a different folder (which is not allowed)
 				else
 				{
 					// Warn the user that the file will not be copied in there
-					OpenMsgDlgInt(
+					const FText Title = LOCTEXT("SaveAsFailedMsg_Title", "Save As Failed");
+					FMessageDialog::Open(
 						EAppMsgType::Ok,
 						FText::Format(
 							LOCTEXT("SaveAsFailedMsg",
@@ -680,7 +773,7 @@ void SaveExportLayoutCommon(const FString& InDefaultDirectory, const bool bMustB
 								"If you simply wish to export a copy of the current configuration in {1} (e.g., to later copy it into a different machine), you could use the \"Export Layout...\""
 								" functionality. However, Unreal would not be able to load it until you import it with \"Import Layout...\"."),
 							FText::FromString(DefaultDirectoryAbsolute), FText::FromString(LayoutFilePathAbsolute)),
-						LOCTEXT("SaveAsFailedMsg_Title", "Save As Failed"));
+						&Title);
 				}
 			}
 		}
@@ -741,10 +834,7 @@ void FLayoutsMenuSave::OverrideUserLayout(const int32 InLayoutIndex)
 	const TArray<FString> UserLayoutIniFileNames = GetIniFilesInFolderInternal(UserLayoutsDirectory);
 	const FString DesiredUserLayoutFullPath = FPaths::Combine(FPaths::GetPath(UserLayoutsDirectory), UserLayoutIniFileNames[InLayoutIndex]);
 	// Are you sure you want to do this?
-	const FText TextFileNameToRemove = FText::FromString(FPaths::GetBaseFilename(UserLayoutIniFileNames[InLayoutIndex]));
-	const FText TextBody = FText::Format(LOCTEXT("ActionOverrideLayoutMsg", "Are you sure you want to permanently override the layout profile \"{0}\" with the current layout profile? This action cannot be undone."), TextFileNameToRemove);
-	const FText TextTitle = FText::Format(LOCTEXT("OverrideUILayout_Title", "Override UI Layout \"{0}\""), TextFileNameToRemove);
-	if (EAppReturnType::Ok != OpenMsgDlgInt(EAppMsgType::OkCancel, TextBody, TextTitle))
+	if (!FSaveLayoutDialogUtils::OverrideLayoutDialog(UserLayoutIniFileNames[InLayoutIndex]))
 	{
 		return;
 	}
@@ -756,7 +846,8 @@ void FLayoutsMenuSave::OverrideUserLayout(const int32 InLayoutIndex)
 	// Replace desired layout with current one
 	const bool bCleanLayoutNameAndDescriptionFieldsIfNoSameValues = true;
 	const bool bShouldAskBeforeCleaningLayoutNameAndDescriptionFields = false;
-	TrySaveLayoutOrWarnInternal(SourceFilePath, TargetFilePath, LOCTEXT("OverrideLayoutText", "layout override"), bCleanLayoutNameAndDescriptionFieldsIfNoSameValues, bShouldAskBeforeCleaningLayoutNameAndDescriptionFields);
+	const bool bShowSaveToast = true;
+	TrySaveLayoutOrWarnInternal(SourceFilePath, TargetFilePath, LOCTEXT("OverrideLayoutText", "overriding the layout"), bCleanLayoutNameAndDescriptionFieldsIfNoSameValues, bShouldAskBeforeCleaningLayoutNameAndDescriptionFields, bShowSaveToast);
 }
 
 void FLayoutsMenuSave::SaveLayout()
@@ -780,7 +871,7 @@ void FLayoutsMenuSave::SaveLayoutAs()
 	const FString DefaultDirectory = CreateAndGetUserLayoutDirInternal();
 	const bool bMustBeSavedInDefaultDirectory = true;
 	const bool bShouldAskBeforeCleaningLayoutNameAndDescriptionFields = false;
-	SaveExportLayoutCommon(DefaultDirectory, bMustBeSavedInDefaultDirectory, LOCTEXT("SaveLayoutText", "layout save"), bShouldAskBeforeCleaningLayoutNameAndDescriptionFields);
+	SaveExportLayoutCommon(DefaultDirectory, bMustBeSavedInDefaultDirectory, LOCTEXT("SaveLayoutText", "saving the layout"), bShouldAskBeforeCleaningLayoutNameAndDescriptionFields);
 }
 
 void FLayoutsMenuSave::ExportLayout()
@@ -791,7 +882,7 @@ void FLayoutsMenuSave::ExportLayout()
 	const FString DefaultDirectory = FPaths::ProjectContentDir();
 	const bool bMustBeSavedInDefaultDirectory = false;
 	const bool bShouldAskBeforeCleaningLayoutNameAndDescriptionFields = true;
-	SaveExportLayoutCommon(DefaultDirectory, bMustBeSavedInDefaultDirectory, LOCTEXT("ExportLayoutText", "layout export"), bShouldAskBeforeCleaningLayoutNameAndDescriptionFields);
+	SaveExportLayoutCommon(DefaultDirectory, bMustBeSavedInDefaultDirectory, LOCTEXT("ExportLayoutText", "exporting the layout"), bShouldAskBeforeCleaningLayoutNameAndDescriptionFields);
 }
 
 
@@ -846,7 +937,7 @@ void FLayoutsMenuRemove::RemoveUserLayout(const int32 InLayoutIndex)
 	const FText TextFileNameToRemove = FText::FromString(FPaths::GetBaseFilename(UserLayoutIniFileNames[InLayoutIndex]));
 	const FText TextBody = FText::Format(LOCTEXT("ActionRemoveMsg", "Are you sure you want to permanently delete the layout profile \"{0}\"? This action cannot be undone."), TextFileNameToRemove);
 	const FText TextTitle = FText::Format(LOCTEXT("RemoveUILayout_Title", "Remove UI Layout \"{0}\""), TextFileNameToRemove);
-	if (EAppReturnType::Ok != OpenMsgDlgInt(EAppMsgType::OkCancel, TextBody, TextTitle))
+	if (EAppReturnType::Ok != FMessageDialog::Open(EAppMsgType::OkCancel, TextBody, &TextTitle))
 	{
 		return;
 	}
@@ -885,7 +976,7 @@ void FLayoutsMenuRemove::RemoveUserLayouts()
 		// Are you sure you want to do this?
 		const FText TextBody = FText::Format(LOCTEXT("ActionRemoveAllUserLayoutMsg", "Are you sure you want to permanently remove {0} layout {0}|plural(one=profile,other=profiles)? This action cannot be undone."), NumberUserLayoutFiles);
 		const FText TextTitle = LOCTEXT("RemoveAllUserLayouts_Title", "Remove All User-Created Layouts");
-		if (EAppReturnType::Ok != OpenMsgDlgInt(EAppMsgType::OkCancel, TextBody, TextTitle))
+		if (EAppReturnType::Ok != FMessageDialog::Open(EAppMsgType::OkCancel, TextBody, &TextTitle))
 		{
 			return;
 		}
@@ -909,7 +1000,7 @@ void FLayoutsMenuRemove::RemoveUserLayouts()
 		// Show reason
 		const FText TextBody = LOCTEXT("UnsuccessfulRemoveLayoutBody", "There are no layout profile files created by the user, so none could be removed.");
 		const FText TextTitle = LOCTEXT("UnsuccessfulRemoveLayoutHeader", "Unsuccessful Remove All User Layouts!");
-		OpenMsgDlgInt(EAppMsgType::Ok, TextBody, TextTitle);
+		FMessageDialog::Open(EAppMsgType::Ok, TextBody, &TextTitle);
 	}
 }
 
diff --git a/Engine/Source/Editor/MainFrame/Private/Menus/SaveLayoutDialog.cpp b/Engine/Source/Editor/MainFrame/Private/Menus/SaveLayoutDialog.cpp
new file mode 100644
--- /dev/null
+++ b/Engine/Source/Editor/MainFrame/Private/Menus/SaveLayoutDialog.cpp
@@ -0,0 +1,505 @@
+// Copyright 1998-2019 Epic Games, Inc. All Rights Reserved.
+
+#include "SaveLayoutDialog.h"
+#include "EditorStyleSet.h"
+#include "Framework/Application/SlateApplication.h"
+#include "Interfaces/IMainFrameModule.h"
+#include "Misc/FileHelper.h"
+#include "Misc/MessageDialog.h"
+#include "SlateOptMacros.h"
+#include "Widgets/DeclarativeSyntaxSupport.h"
+#include "Widgets/Input/SButton.h"
+#include "Widgets/Input/SEditableTextBox.h"
+#include "Widgets/Layout/SBox.h"
+#include "Widgets/SBoxPanel.h"
+#include "Widgets/SCompoundWidget.h"
+#include "Widgets/Text/STextBlock.h"
+
+#define LOCTEXT_NAMESPACE "SaveLayoutDialog"
+
+/**
+ * Note: Code for SSaveLayoutDialog obtained and modified from SAssetDialog
+ *
+ * A SWidget that generates a custom SCompoundWidget that allows the user to select the final file name where a layout ini file will be saved,
+ * and optionally the displayed layout name and description.
+ *
+ * This SWidget is never meant to used by itself. Instead, FSaveLayoutDialogUtils::CreateSaveLayoutAsDialogInStandaloneWindow() should be called,
+ * which additionally sets the right SWindow for SSaveLayoutDialog and blocks the thread until the user finishes interacting with it.
+ */
+class SSaveLayoutDialog : public SCompoundWidget
+{
+public:
+	SLATE_BEGIN_ARGS(SSaveLayoutDialog) {}
+
+	SLATE_END_ARGS()
+
+	SSaveLayoutDialog();
+
+	virtual ~SSaveLayoutDialog();
+
+	virtual void Construct(const FArguments& InArgs, const TSharedRef<FSaveLayoutDialogParams>& InSaveLayoutDialogParams);
+
+private:
+	FText GetFileNameText() const;
+
+	/** Gets the visibility of the name error label */
+	EVisibility GetNameErrorLabelVisibility() const;
+
+	/** Gets the text to display in the name error label */
+	FText GetNameErrorLabelText() const;
+
+	/** Fired when the layout name box text is committed */
+	void OnLayoutNameTextCommited(const FText& InText, ETextCommit::Type InCommitType);
+
+	/** Fired when the layout description box text is committed */
+	void OnLayoutDescriptionTextCommited(const FText& InText, ETextCommit::Type InCommitType);
+
+	/** Determines if the Save button is enabled. */
+	bool IsConfirmButtonEnabled() const;
+
+	/** Handler for when the Save button is clicked */
+	FReply OnConfirmClicked();
+
+	/** Handler for when the cancel button is clicked */
+	FReply OnCancelClicked();
+
+	/** Used to commit the object path used for saving in this dialog */
+	void CommitObjectPathForSave();
+
+	/** Closes this dialog */
+	void CloseDialog();
+
+	bool IsValidFilePathForCreation(const FString& FilePath, const FString& FileName, FText& OutErrorMessage, bool bAllowExistingAsset = true) const;
+
+	void UpdateInputValidity();
+
+	TSharedPtr<FSaveLayoutDialogParams> SaveLayoutDialogParams;
+
+	FString DefaultDirectory;
+
+	FString FinalFileName;
+
+	FString FinalFilePath;
+
+	FText CurrentlyEnteredLayoutName;
+
+	FText CurrentlyEnteredLayoutDescription;
+
+	/** The error text from the last validity check */
+	FText LastInputValidityErrorText;
+
+	/** True if the last validity check returned that the class name/path is valid for creation */
+	bool bLastInputValidityCheckSuccessful;
+};
+
+
+
+FSaveLayoutDialogParams::FSaveLayoutDialogParams(const FString& InDefaultDirectory, const FString& InFileExtension, const TArray<FText>& InLayoutNames, const TArray<FText>& InLayoutDescriptions)
+	: DefaultDirectory(InDefaultDirectory)
+	, FileExtension(InFileExtension)
+	, LayoutNames(InLayoutNames)
+	, LayoutDescriptions(InLayoutDescriptions)
+{
+}
+
+void FSaveLayoutDialogUtils::SanitizeText(FString& InOutString)
+{
+	// Replace any special character by `_`
+	for (TCHAR& InOutChar : InOutString)
+	{
+		if (!FChar::IsAlnum(InOutChar))
+		{
+			InOutChar = '_';
+		}
+	}
+}
+
+bool FSaveLayoutDialogUtils::OverrideLayoutDialog(const FString& LayoutIniFileName)
+{
+	// Are you sure you want to do this?
+	const FText TextFileNameToRemove = FText::FromString(FPaths::GetBaseFilename(LayoutIniFileName));
+	const FText TextBody = FText::Format(
+		LOCTEXT("ActionOverrideLayoutMsg", "Are you sure you want to permanently override the layout profile \"{0}\" with the current layout profile? This action cannot be undone."), TextFileNameToRemove);
+	const FText TextTitle = FText::Format(LOCTEXT("OverrideUILayout_Title", "Override UI Layout \"{0}\"?"), TextFileNameToRemove);
+	return (EAppReturnType::Ok == FMessageDialog::Open(EAppMsgType::OkCancel, TextBody, &TextTitle));
+}
+
+bool FSaveLayoutDialogUtils::CreateSaveLayoutAsDialogInStandaloneWindow(const TSharedRef<FSaveLayoutDialogParams>& InSaveLayoutDialogParams)
+{
+	// Create SSaveLayoutDialog
+	TSharedRef<SSaveLayoutDialog> SaveLayoutDialog = SNew(SSaveLayoutDialog, InSaveLayoutDialogParams);
+
+	// Create SWindow that contains SSaveLayoutDialog
+	const FVector2D WindowSize(720.f, 130.f);
+	TSharedRef<SWindow> DialogWindow =
+		SNew(SWindow)
+		.Title(LOCTEXT("GenericAssetDialogWindowHeader", "Save Layout As"))
+		.ClientSize(WindowSize);
+	DialogWindow->SetContent(SaveLayoutDialog);
+
+	// Launch SSaveLayoutDialog and block thread until user finishes with it
+	IMainFrameModule& MainFrameModule = FModuleManager::LoadModuleChecked<IMainFrameModule>(TEXT("MainFrame"));
+	const TSharedPtr<SWindow>& MainFrameParentWindow = MainFrameModule.GetParentWindow();
+	if (MainFrameParentWindow.IsValid())
+	{
+		// Opposite to AddWindowAsNativeChild(), AddModalWindow() does not return until the modal window is closed
+		FSlateApplication::Get().AddModalWindow(DialogWindow, MainFrameParentWindow.ToSharedRef());
+		return true;
+	}
+	else
+	{
+		ensureMsgf(false, TEXT("Could not create \"Save Layout As...\" dialog, this should never happen."));
+		return false;
+	}
+}
+
+SSaveLayoutDialog::SSaveLayoutDialog()
+	: bLastInputValidityCheckSuccessful(false)
+{
+}
+
+SSaveLayoutDialog::~SSaveLayoutDialog()
+{
+}
+
+BEGIN_SLATE_FUNCTION_BUILD_OPTIMIZATION
+void SSaveLayoutDialog::Construct(const FArguments& InArgs, const TSharedRef<FSaveLayoutDialogParams>& InSaveLayoutDialogParams)
+{
+	SaveLayoutDialogParams = InSaveLayoutDialogParams;
+
+	// Set Default Saving Directory
+	DefaultDirectory = SaveLayoutDialogParams->DefaultDirectory;
+
+	// Fill initial (default) values
+	if (SaveLayoutDialogParams->LayoutNames.Num() > 0)
+	{
+		ensureMsgf(SaveLayoutDialogParams->LayoutNames.Num() == 1, TEXT("This code is only ready for SaveLayoutDialogParams->LayoutNames.Num() == 1."));
+		CurrentlyEnteredLayoutName = SaveLayoutDialogParams->LayoutNames[0];
+		UpdateInputValidity();
+	}
+	if (SaveLayoutDialogParams->LayoutDescriptions.Num() > 0)
+	{
+		ensureMsgf(SaveLayoutDialogParams->LayoutDescriptions.Num() == 1, TEXT("This code is only ready for LayoutNames.Num() == 1."));
+		CurrentlyEnteredLayoutDescription = SaveLayoutDialogParams->LayoutDescriptions[0];
+	}
+
+	// UI part
+	// The root SWidget in this dialog
+	TSharedRef<SVerticalBox> MainVerticalBox = SNew(SVerticalBox);
+	TSharedRef<SVerticalBox> LabelsBox = SNew(SVerticalBox);
+	TSharedRef<SVerticalBox> ContentBox = SNew(SVerticalBox);
+
+	// Layout Name
+	LabelsBox->AddSlot()
+		.FillHeight(1)
+		.VAlign(VAlign_Center)
+		.Padding(0, 3.625f, 0, 3.625f)
+		[
+			SNew(STextBlock).Text(LOCTEXT("LayoutNameBoxLabel", "Name:"))
+		];
+	ContentBox->AddSlot()
+		.FillHeight(1)
+		.VAlign(VAlign_Center)
+		.Padding(0, 2.f, 0, 2.f)
+		[
+			SNew(SEditableTextBox)
+			.Text(CurrentlyEnteredLayoutName)
+			.OnTextCommitted(this, &SSaveLayoutDialog::OnLayoutNameTextCommited)
+			.OnTextChanged(this, &SSaveLayoutDialog::OnLayoutNameTextCommited, ETextCommit::Default)
+			.SelectAllTextWhenFocused(true)
+		];
+
+	// Layout Description
+	LabelsBox->AddSlot()
+		.FillHeight(1)
+		.VAlign(VAlign_Center)
+		.Padding(0, 3.625f, 0, 3.625f)
+		[
+			SNew(STextBlock).Text(LOCTEXT("LayoutDescriptionBoxLabel", "Description (Optional):"))
+		];
+	ContentBox->AddSlot()
+		.FillHeight(1)
+		.VAlign(VAlign_Center)
+		.Padding(0, 2.f, 0, 2.f)
+		[
+			SNew(SEditableTextBox)
+			.Text(CurrentlyEnteredLayoutDescription)
+			.OnTextCommitted(this, &SSaveLayoutDialog::OnLayoutDescriptionTextCommited)
+			.OnTextChanged(this, &SSaveLayoutDialog::OnLayoutDescriptionTextCommited, ETextCommit::Default)
+			.SelectAllTextWhenFocused(true)
+		];
+
+	// Buttons and asset name
+	TSharedRef<SHorizontalBox> LabelAndContentBox = SNew(SHorizontalBox)
+		+ SHorizontalBox::Slot()
+		.AutoWidth()
+		.HAlign(HAlign_Right)
+		.VAlign(VAlign_Bottom)
+		.Padding(30, 30, 4, 3)
+		[
+			LabelsBox
+		]
+		+ SHorizontalBox::Slot()
+		.FillWidth(1)
+		.VAlign(VAlign_Bottom)
+		.Padding(4, 3, 30, 4)
+		[
+			ContentBox
+		];
+	MainVerticalBox->AddSlot()
+		.AutoHeight()
+		.HAlign(HAlign_Fill)
+		.Padding(0)
+		[
+			LabelAndContentBox
+		];
+
+	// Name Error label
+	MainVerticalBox->AddSlot()
+	.AutoHeight()
+	[
+		// Constant height, whether the label is visible or not
+		SNew(SBox).HeightOverride(20)
+		[
+			SNew(SBorder)
+			.Visibility( this, &SSaveLayoutDialog::GetNameErrorLabelVisibility )
+			.BorderImage( FEditorStyle::GetBrush("AssetDialog.ErrorLabelBorder") )
+			.Content()
+			[
+				SNew(STextBlock)
+				.Text( this, &SSaveLayoutDialog::GetNameErrorLabelText )
+				.TextStyle( FEditorStyle::Get(), "AssetDialog.ErrorLabelFont" )
+			]
+		]
+	];
+
+	// Buttons
+	TSharedRef<SHorizontalBox> ButtonBox = SNew(SHorizontalBox)
+		+ SHorizontalBox::Slot()
+		.AutoWidth()
+		.VAlign(VAlign_Bottom)
+		.Padding(4, 3)
+		[
+			SNew(SButton)
+			.Text(LOCTEXT("SaveLayoutDialogSaveButton", "Save"))
+			.ContentPadding(FMargin(8, 2, 8, 2))
+			.IsEnabled(this, &SSaveLayoutDialog::IsConfirmButtonEnabled)
+			.OnClicked(this, &SSaveLayoutDialog::OnConfirmClicked)
+		]
+		+ SHorizontalBox::Slot()
+		.AutoWidth()
+		.VAlign(VAlign_Bottom)
+		.Padding(4, 3)
+		[
+			SNew(SButton)
+			.ContentPadding(FMargin(8, 2, 8, 2))
+			.Text(LOCTEXT("SaveLayoutDialogCancelButton", "Cancel"))
+			.OnClicked(this, &SSaveLayoutDialog::OnCancelClicked)
+		];
+	MainVerticalBox->AddSlot()
+		.HAlign(HAlign_Right)
+		.VAlign(VAlign_Bottom)
+		.Padding(0, 0, 5, 5)
+		[
+			ButtonBox
+		];
+
+	// Add to ChildSlot
+	ChildSlot
+	[
+		MainVerticalBox
+	];
+}
+END_SLATE_FUNCTION_BUILD_OPTIMIZATION
+
+FText SSaveLayoutDialog::GetFileNameText() const
+{
+	return FText::FromString(FinalFileName);
+}
+
+EVisibility SSaveLayoutDialog::GetNameErrorLabelVisibility() const
+{
+	return GetNameErrorLabelText().IsEmpty() ? EVisibility::Hidden : EVisibility::Visible;
+}
+
+FText SSaveLayoutDialog::GetNameErrorLabelText() const
+{
+	if (!bLastInputValidityCheckSuccessful)
+	{
+		return LastInputValidityErrorText;
+	}
+
+	return FText::GetEmpty();
+}
+
+void SSaveLayoutDialog::OnLayoutNameTextCommited(const FText& InText, ETextCommit::Type InCommitType)
+{
+	CurrentlyEnteredLayoutName = InText;
+	UpdateInputValidity();
+
+	if (InCommitType == ETextCommit::OnEnter)
+	{
+		CommitObjectPathForSave();
+	}
+}
+
+void SSaveLayoutDialog::OnLayoutDescriptionTextCommited(const FText& InText, ETextCommit::Type InCommitType)
+{
+	CurrentlyEnteredLayoutDescription = InText;
+
+	if (InCommitType == ETextCommit::OnEnter)
+	{
+		CommitObjectPathForSave();
+	}
+}
+
+bool SSaveLayoutDialog::IsConfirmButtonEnabled() const
+{
+	return bLastInputValidityCheckSuccessful;
+}
+
+FReply SSaveLayoutDialog::OnConfirmClicked()
+{
+	CommitObjectPathForSave();
+	return FReply::Handled();
+}
+
+FReply SSaveLayoutDialog::OnCancelClicked()
+{
+	bLastInputValidityCheckSuccessful = false;
+	CloseDialog();
+	return FReply::Handled();
+}
+
+void SSaveLayoutDialog::CommitObjectPathForSave()
+{
+	// Update final values and close dialog if valid input and not duplicated
+	if (bLastInputValidityCheckSuccessful)
+	{
+		// Sanity check
+		if (SaveLayoutDialogParams.IsValid())
+		{
+			// Are you sure you want to do this?
+			if (FPaths::FileExists(FinalFilePath) && !FSaveLayoutDialogUtils::OverrideLayoutDialog(FinalFilePath))
+			{
+				return;
+			}
+			// Update output parameters
+			SaveLayoutDialogParams->bWereFilesSelected = bLastInputValidityCheckSuccessful;
+			SaveLayoutDialogParams->LayoutFilePaths.Empty();
+			SaveLayoutDialogParams->LayoutNames.Empty();
+			SaveLayoutDialogParams->LayoutDescriptions.Empty();
+			if (SaveLayoutDialogParams->bWereFilesSelected)
+			{
+				SaveLayoutDialogParams->LayoutFilePaths.Emplace(FinalFilePath);
+				SaveLayoutDialogParams->LayoutNames.Emplace(CurrentlyEnteredLayoutName);
+				SaveLayoutDialogParams->LayoutDescriptions.Emplace(CurrentlyEnteredLayoutDescription);
+			}
+		}
+		else
+		{
+			ensureMsgf(false, TEXT("SaveLayoutDialogParams should never be a nullptr at this point, it is ideally initialized in the Construct() class."));
+		}
+		// Close dialog
+		CloseDialog();
+	}
+}
+
+void SSaveLayoutDialog::CloseDialog()
+{
+	TSharedPtr<SWindow> ContainingWindow = FSlateApplication::Get().FindWidgetWindow(AsShared());
+
+	if (ContainingWindow.IsValid())
+	{
+		ContainingWindow->RequestDestroyWindow();
+	}
+}
+
+bool SSaveLayoutDialog::IsValidFilePathForCreation(const FString& FilePath, const FString& FileName, FText& OutErrorMessage, bool bAllowExistingAsset) const
+{
+	// Make sure the name is not already a class or otherwise invalid for saving
+	if (!FFileHelper::IsFilenameValidForSaving(FileName, OutErrorMessage))
+	{
+		// Return false to indicate that the user should enter a new name
+		return false;
+	}
+
+	// Make sure the new name only contains valid characters
+	if (!FName::IsValidXName(FileName, INVALID_OBJECTNAME_CHARACTERS INVALID_LONGPACKAGE_CHARACTERS, &OutErrorMessage))
+	{
+		// Return false to indicate that the user should enter a new name
+		return false;
+	}
+
+	// Make sure we are not creating an FName that is too large
+	if (FilePath.Len() > NAME_SIZE)
+	{
+		// This asset already exists at this location, inform the user and continue
+		OutErrorMessage = LOCTEXT("AssetNameTooLong", "This asset name is too long. Please choose a shorter name.");
+		// Return false to indicate that the user should enter a new name
+		return false;
+	}
+
+	// Make sure we are not creating an path that is too long for the OS
+	const FString FullPath = FPaths::ConvertRelativePathToFull(FilePath);
+	if (FilePath.Len() > FPlatformMisc::GetMaxPathLength())
+	{
+		// The full path for the asset is too long
+		OutErrorMessage = FText::Format(LOCTEXT("ObjectPathTooLong",
+			"The object path for the asset is too long, the maximum is '{0}'. \nPlease choose a shorter name for the asset or create it in a shallower folder structure."),
+			FText::AsNumber(FPlatformMisc::GetMaxPathLength()));
+		// Return false to indicate that the user should enter a new name
+		return false;
+	}
+
+	// Check for an existing asset, unless it we were asked not to.
+	if (!bAllowExistingAsset)
+	{
+		if (FPaths::FileExists(FilePath))
+		{
+			// This asset already exists at this location, inform the user and continue
+			OutErrorMessage = FText::Format(LOCTEXT("RenameAssetAlreadyExists", "An asset already exists at this location with the name '{0}'."), FText::FromString(FileName));
+
+			// Return false to indicate that the user should enter a new name
+			return false;
+		}
+	}
+
+	return true;
+}
+
+void SSaveLayoutDialog::UpdateInputValidity()
+{
+	// Update file name
+	FinalFileName = CurrentlyEnteredLayoutName.ToString();
+	FSaveLayoutDialogUtils::SanitizeText(FinalFileName);
+	// Fill full layout file path from file name and directory
+	FinalFilePath =
+		// Folder
+		(DefaultDirectory.Right(1) == TEXT("/") ? DefaultDirectory : DefaultDirectory + TEXT("/"))
+		// FileName + Extension
+		+ (FinalFileName.Len() > SaveLayoutDialogParams->FileExtension.Len() && FinalFileName.Right(3) == SaveLayoutDialogParams->FileExtension
+			? FinalFileName : FinalFileName + SaveLayoutDialogParams->FileExtension);
+
+	// Validate file name
+	bLastInputValidityCheckSuccessful = true;
+
+	// No error text for an empty name. Just fail validity.
+	if (FinalFileName.IsEmpty())
+	{
+		bLastInputValidityCheckSuccessful = false;
+		LastInputValidityErrorText = FText::GetEmpty();
+	}
+
+	if (bLastInputValidityCheckSuccessful)
+	{
+		FText ErrorMessage;
+		if (!IsValidFilePathForCreation(FinalFilePath, FinalFileName, ErrorMessage))
+		{
+			bLastInputValidityCheckSuccessful = false;
+			LastInputValidityErrorText = ErrorMessage;
+		}
+	}
+}
+
+#undef LOCTEXT_NAMESPACE
diff --git a/Engine/Source/Editor/MainFrame/Private/Menus/SaveLayoutDialog.h b/Engine/Source/Editor/MainFrame/Private/Menus/SaveLayoutDialog.h
new file mode 100644
--- /dev/null
+++ b/Engine/Source/Editor/MainFrame/Private/Menus/SaveLayoutDialog.h
@@ -0,0 +1,76 @@
+// Copyright 1998-2019 Epic Games, Inc. All Rights Reserved.
+
+#pragma once
+
+#include "CoreMinimal.h"
+
+/**
+ * It saves the input parameters that FSaveLayoutDialogUtils::CreateSaveLayoutAsDialogInStandaloneWindow needs and the output ones that it will generate.
+ */
+class FSaveLayoutDialogParams
+{
+public:
+	/**
+	 * The constructor only contains the variables that FSaveLayoutDialogUtils::CreateSaveLayoutAsDialogInStandaloneWindow will need. The rest will be overriden.
+	 */
+	FSaveLayoutDialogParams(const FString& InDefaultDirectory, const FString& InFileExtension, const TArray<FText>& LayoutNames = TArray<FText>(),
+		const TArray<FText>& LayoutDescriptions = TArray<FText>());
+
+	/**
+	 * The directory where the layout will be saved.
+	 */
+	const FString DefaultDirectory;
+
+	/**
+	 * The extension of the layout file.
+	 */
+	const FString FileExtension;
+
+	/**
+	 * Whether the user selected files (true) or cancelled/closed the window (false)
+	 */
+	bool bWereFilesSelected;
+
+	/**
+	 * The final full file path where the layout profile should be saved
+	 */
+	TArray<FString> LayoutFilePaths;
+
+	/**
+	 * The displayed layout name that will be visualized in Unreal. Its initial value will be used as default.
+	 */
+	TArray<FText> LayoutNames;
+
+	/**
+	 * The displayed layout description that will be visualized in Unreal. Its initial value will be used as default.
+	 */
+	TArray<FText> LayoutDescriptions;
+};
+
+/**
+ * Some convenient UI functions regarding SWidgets for saving the layout
+ */
+class FSaveLayoutDialogUtils
+{
+public:
+	/**
+	 * Any special character (i.e., any non-alphanumeric character) will be turned into an underscore.
+	 */
+	static void SanitizeText(FString& InOutString);
+
+	/**
+	 * It will show a text asking whether you wanna override the existing UI Layout ini file.
+	 * @param LayoutIniFileName The file path that might be overridden.
+	 * @return True if the user agrees on overriding the file, false otherwise.
+	 */
+	static bool OverrideLayoutDialog(const FString& LayoutIniFileName);
+
+	/**
+	 * Self-contained function that creates a standalone SWindow, the SSaveLayoutDialog widget, and that blocks the thread until the dialog is finished
+	 * @param InSaveLayoutDialogParams The FSaveLayoutDialogParams parameters used in the construction of SSaveLayoutDialog. Looking at its values:
+	 *     1) DefaultDirectory and FileExtension will only be read.
+	 *     2) bWereFilesSelected and LayoutFilePaths will only be written (i.e., their initial value will be overridden).
+	 *     3) LayoutNames and LayoutDescriptions will only be written, but their initial values will also be used as default.
+	 */
+	static bool CreateSaveLayoutAsDialogInStandaloneWindow(const TSharedRef<FSaveLayoutDialogParams>& InSaveLayoutDialogParams);
+};
