Copied CL# 10108258

Bug fixed: Loading an Editor Layout used to close all open sub-editors.  Now, a new Dialog warning window is displayed if any sub-editor instance or PIE/SIE are opened when the user tries to load a new layout, asking whether he want to continue loading the layout (that would lead to automatically close+- those without saving them).

Note: Previously, all the load options were disabled while PIE was running. However, this could be confusing for the user, who might not know why they are disabled and not associate it with PIE being running (Although it was shown in the Load Layout tooltip, but people is not likely to read those...). Instead, this warning is merged with the one for the opened asset editors.

This is a quick fix is for 4.24, to avoid deep code changes. However, for later versions, the following tasks should be done as well:
- Investigate whether PIE or the sub-editor can be maintained opened, rather than closing them when a new layout is loaded.

diff --git a/Engine/Source/Editor/MainFrame/Private/Menus/LayoutsMenu.cpp b/Engine/Source/Editor/MainFrame/Private/Menus/LayoutsMenu.cpp
--- a/Engine/Source/Editor/MainFrame/Private/Menus/LayoutsMenu.cpp
+++ b/Engine/Source/Editor/MainFrame/Private/Menus/LayoutsMenu.cpp
@@ -31,6 +31,7 @@
 #include "LevelViewportActions.h"
 #include "Menus/SaveLayoutDialog.h"
 #include "Misc/MessageDialog.h"
+#include "Subsystems/AssetEditorSubsystem.h"
 #include "UnrealEdGlobals.h"
 #include "UnrealEdMisc.h"
 
@@ -463,6 +464,88 @@ void SaveLayoutWithoutRemovingTempLayoutFiles()
 	GConfig->Flush(bRead, GEditorLayoutIni);
 }
 
+/**
+ * It simply check whether PIE, SIE, or any Asset Editor is opened, and ask the user whether he wanna continue closing them or cancel the Editor layout load
+ * @return Whether we should continue loading the layout
+ */
+bool CheckAskUserAndClosePIESIEAndAssetEditors(const FText& InitialMessage)
+{
+	// If none are running, return
+	const bool bIsPIEOrSIERunning = ((GEditor && GEditor->PlayWorld) || GIsPlayInEditorWorld);
+	UAssetEditorSubsystem* AssetEditorSubsystem = GEditor->GetEditorSubsystem<UAssetEditorSubsystem>();
+	const TArray<UObject*> AllEditedAssets = AssetEditorSubsystem->GetAllEditedAssets();
+	const bool bAreAssetEditorOpened = (AllEditedAssets.Num() > 0);
+	if (!bIsPIEOrSIERunning && !bAreAssetEditorOpened)
+	{
+		return true;
+	}
+	// Collect all open assets
+	FText OpenedEditorAssets;
+	if (bAreAssetEditorOpened)
+	{
+		FString AllAssets;
+		for (const UObject* EditedAsset : AllEditedAssets)
+		{
+			if (!AllAssets.IsEmpty())
+				AllAssets += TEXT(", ");
+			AllAssets += EditedAsset->GetName();
+		}
+		OpenedEditorAssets = FText::Format(LOCTEXT("CheckAskUserAndClosePIESIEAndAssetEditorsOpenEditorAssets", "\nOpen Asset Editors: {0}."), FText::FromString(AllAssets));
+	}
+	else
+	{
+		OpenedEditorAssets = LOCTEXT("CheckAskUserAndClosePIESIEAndAssetEditorsOpenEditorAssetsEmpty", "\n");
+	}
+	FText TextTitle;
+	FText IfYesText;
+	// If both PIE/SIE and Asset Editors are opened
+	if (bIsPIEOrSIERunning && bAreAssetEditorOpened)
+	{
+		TextTitle = LOCTEXT("CheckAskUserAndClosePIESIEAndAssetEditorsIfYesHeaderAll", "Close PIE/SIE and Asset Editors?");
+		IfYesText = LOCTEXT("CheckAskUserAndClosePIESIEAndAssetEditorsIfYesBodyAll", "If \"Yes\", your current game instances (PIE or SIE) as well as all open Asset Editors will be closed. Any unsaved changes in those will also be lost.");
+	}
+	// If PIE or SIE are opened
+	else if (bIsPIEOrSIERunning)
+	{
+		TextTitle = LOCTEXT("CheckAskUserAndClosePIESIEAndAssetEditorsIfYesHeaderPIE", "Close PIE/SIE?");
+		IfYesText = LOCTEXT("CheckAskUserAndClosePIESIEAndAssetEditorsIfYesBodyPIE", "If \"Yes\", your current game instances (PIE or SIE) will be closed. Any unsaved changes in those will also be lost.");
+	}
+	// If any Asset Editors is opened
+	else if (bAreAssetEditorOpened)
+	{
+		TextTitle = LOCTEXT("CheckAskUserAndClosePIESIEAndAssetEditorsIfYesHeaderEditorAssets", "Close Asset Editors?");
+		IfYesText = LOCTEXT("CheckAskUserAndClosePIESIEAndAssetEditorsIfYesBodyEditorAssets", "If \"Yes\", all open Asset Editors will be closed. Any unsaved changes in those will also be lost.");
+	}
+	// FMessageDialog
+	const FText IfNoText = LOCTEXT("CheckAskUserAndClosePIESIEAndAssetEditorsIfNoBody", "If \"No\" or \"Cancel\", you can manually reload the layout from the \"User Layouts\" section later.");
+	const FText TextBody = FText::Format(LOCTEXT("ClosePIESIEAssetEditorsBody", "{0}\n\n{1}{2}\n\n{3}"), InitialMessage, IfYesText, OpenedEditorAssets, IfNoText);
+	if (EAppReturnType::Yes != FMessageDialog::Open(EAppMsgType::YesNoCancel, TextBody, &TextTitle))
+	{
+		return false;
+	}
+	// If PIE or SIE are opened, ask the user whether he wants to automatically close them and continue loading the layout
+	if (bIsPIEOrSIERunning)
+	{
+		// Close PIE/SIE
+		if (GEditor && GEditor->PlayWorld)
+		{
+			GEditor->EndPlayMap();
+		}
+		else
+		{
+			ensureMsgf(false,
+				TEXT("This has not been tested because the code does not reach this by default. The layout is loaded through the Editor UI, and GIsPlayInEditorWorld should not have any kind of Editor UI, so it should not be possible to load a layout in that status."));
+		}
+	}
+	// If any Asset Editors is opened, ask the user whether he wants to automatically close them and continue loading the layout
+	if (bAreAssetEditorOpened)
+	{
+		// Close asset editors
+		AssetEditorSubsystem->CloseAllAssetEditors();
+	}
+	return true;
+}
+
 
 
 void FLayoutsMenuLoad::MakeLoadLayoutsMenu(UToolMenu* InToolMenu)
@@ -491,23 +574,20 @@ void FLayoutsMenuLoad::MakeLoadLayoutsMenu(UToolMenu* InToolMenu)
 
 bool FLayoutsMenuLoad::CanLoadChooseLayout(const int32 InLayoutIndex)
 {
-	if (GEngine && GEngine->GameViewport)
-	{
-		return false;
-	}
 	return !FLayoutsMenuBase::IsLayoutChecked(InLayoutIndex) && CanChooseLayoutWhenReadInternal(InLayoutIndex);
 }
 bool FLayoutsMenuLoad::CanLoadChooseUserLayout(const int32 InLayoutIndex)
 {
-	if (GEngine && GEngine->GameViewport)
-	{
-		return false;
-	}
 	return !FLayoutsMenuBase::IsUserLayoutChecked(InLayoutIndex) && CanChooseUserLayoutWhenReadInternal(InLayoutIndex);
 }
 
 void FLayoutsMenuLoad::ReloadCurrentLayout()
 {
+	// If PIE, SIE, or any Asset Editors are opened, ask the user whether he wants to automatically close them and continue loading the layout
+	if (!CheckAskUserAndClosePIESIEAndAssetEditors(LOCTEXT("AreYouSureToLoadHeader", "Are you sure you want to continue loading the selected layout profile?")))
+	{
+		return;
+	}
 	// Create duplicated ini file (OriginalEditorLayoutIniFilePath)
 	// Explanation:
 	//     Assume a layout is saved with (at least) a window that is dependent on a plugin. If that plugin is disabled and the editor restarted, that window will be saved on the layout but will
@@ -626,13 +706,9 @@ void FLayoutsMenuLoad::ImportLayout()
 					FMessageDialog::Open(EAppMsgType::Ok, TextBody, &TextTitle);
 				}
 			}
-			// If PIE running, do not reload current layout
-			if (GEngine && GEngine->GameViewport)
+			// If PIE, SIE, or any Asset Editors are opened, ask the user whether he wants to automatically close them and continue loading the layout
+			if (!CheckAskUserAndClosePIESIEAndAssetEditors(LOCTEXT("LayoutImportClosePIEAndEditorAssetsHeader", "The layout(s) were successfully imported into the \"User Layouts\" section. Do you want to continue loading the selected layout profile?")))
 			{
-				const FText TextBody = LOCTEXT("SuccessfulImportBody",
-					"The layout(s) were successfully imported into the \"User Layouts\" section. However, no layout has been loaded into your current Unreal Editor UI because PIE is currently running. In order to do so, you must stop PIE and then load the layout from the \"User Layouts\" section.");
-				const FText TextTitle = LOCTEXT("SuccessfulImportHeader", "Successful Import!");
-				FMessageDialog::Open(EAppMsgType::Ok, TextBody, &TextTitle);
 				return;
 			}
 			// Replace current layout with first one
@@ -936,7 +1012,7 @@ void FLayoutsMenuRemove::RemoveUserLayout(const int32 InLayoutIndex)
 	// Are you sure you want to do this?
 	const FText TextFileNameToRemove = FText::FromString(FPaths::GetBaseFilename(UserLayoutIniFileNames[InLayoutIndex]));
 	const FText TextBody = FText::Format(LOCTEXT("ActionRemoveMsg", "Are you sure you want to permanently delete the layout profile \"{0}\"? This action cannot be undone."), TextFileNameToRemove);
-	const FText TextTitle = FText::Format(LOCTEXT("RemoveUILayout_Title", "Remove UI Layout \"{0}\""), TextFileNameToRemove);
+	const FText TextTitle = FText::Format(LOCTEXT("RemoveUILayout_Title", "Remove UI Layout \"{0}\"?"), TextFileNameToRemove);
 	if (EAppReturnType::Ok != FMessageDialog::Open(EAppMsgType::OkCancel, TextBody, &TextTitle))
 	{
 		return;
@@ -975,7 +1051,7 @@ void FLayoutsMenuRemove::RemoveUserLayouts()
 	{
 		// Are you sure you want to do this?
 		const FText TextBody = FText::Format(LOCTEXT("ActionRemoveAllUserLayoutMsg", "Are you sure you want to permanently remove {0} layout {0}|plural(one=profile,other=profiles)? This action cannot be undone."), NumberUserLayoutFiles);
-		const FText TextTitle = LOCTEXT("RemoveAllUserLayouts_Title", "Remove All User-Created Layouts");
+		const FText TextTitle = LOCTEXT("RemoveAllUserLayouts_Title", "Remove All User-Created Layouts?");
 		if (EAppReturnType::Ok != FMessageDialog::Open(EAppMsgType::OkCancel, TextBody, &TextTitle))
 		{
 			return;
diff --git a/Engine/Source/Editor/MainFrame/Private/Menus/MainMenu.cpp b/Engine/Source/Editor/MainFrame/Private/Menus/MainMenu.cpp
--- a/Engine/Source/Editor/MainFrame/Private/Menus/MainMenu.cpp
+++ b/Engine/Source/Editor/MainFrame/Private/Menus/MainMenu.cpp
@@ -256,7 +256,7 @@ void FMainMenu::RegisterWindowMenu()
 		Section.AddEntry(FToolMenuEntry::InitSubMenu(
 			"LoadLayout",
 			NSLOCTEXT("LayoutMenu", "LayoutLoadHeader", "Load Layout"),
-			NSLOCTEXT("LayoutMenu", "LoadLayoutsSubMenu_ToolTip", "Load a layout configuration from disk. If PIE is running, most options will be disabled"),
+			NSLOCTEXT("LayoutMenu", "LoadLayoutsSubMenu_ToolTip", "Load a layout configuration from disk"),
 			FNewToolMenuDelegate::CreateStatic(&FLayoutsMenuLoad::MakeLoadLayoutsMenu)
 		));
 		// Save and Remove Layout
