Fixing regression warning created by CL # 10110416

diff --git a/Engine/Source/Editor/MainFrame/Private/Menus/LayoutsMenu.cpp b/Engine/Source/Editor/MainFrame/Private/Menus/LayoutsMenu.cpp
--- a/Engine/Source/Editor/MainFrame/Private/Menus/LayoutsMenu.cpp
+++ b/Engine/Source/Editor/MainFrame/Private/Menus/LayoutsMenu.cpp
@@ -470,9 +470,14 @@ void SaveLayoutWithoutRemovingTempLayoutFiles()
  */
 bool CheckAskUserAndClosePIESIEAndAssetEditors(const FText& InitialMessage)
 {
+	UAssetEditorSubsystem* AssetEditorSubsystem = (GEditor ? GEditor->GetEditorSubsystem<UAssetEditorSubsystem>() : nullptr);
+	if (!GEditor || !AssetEditorSubsystem)
+	{
+		ensureMsgf(false, TEXT("Both GEditor and AssetEditorSubsystem should not be false when CheckAskUserAndClosePIESIEAndAssetEditors() is called."));
+		return true;
+	}
 	// If none are running, return
 	const bool bIsPIEOrSIERunning = ((GEditor && GEditor->PlayWorld) || GIsPlayInEditorWorld);
-	UAssetEditorSubsystem* AssetEditorSubsystem = GEditor->GetEditorSubsystem<UAssetEditorSubsystem>();
 	const TArray<UObject*> AllEditedAssets = AssetEditorSubsystem->GetAllEditedAssets();
 	const bool bAreAssetEditorOpened = (AllEditedAssets.Num() > 0);
 	if (!bIsPIEOrSIERunning && !bAreAssetEditorOpened)
