Copied CL# 10119067

Solved mini-bug when saving a layout with path too long: It was wrongly checking the relative path length rather than the fulll one.

diff --git a/Engine/Source/Editor/MainFrame/Private/Menus/SaveLayoutDialog.cpp b/Engine/Source/Editor/MainFrame/Private/Menus/SaveLayoutDialog.cpp
--- a/Engine/Source/Editor/MainFrame/Private/Menus/SaveLayoutDialog.cpp
+++ b/Engine/Source/Editor/MainFrame/Private/Menus/SaveLayoutDialog.cpp
@@ -432,22 +432,23 @@ bool SSaveLayoutDialog::IsValidFilePathForCreation(const FString& FilePath, cons
 	}
 
 	// Make sure we are not creating an FName that is too large
-	if (FilePath.Len() > NAME_SIZE)
+	if (FilePath.Len() >= NAME_SIZE)
 	{
 		// This asset already exists at this location, inform the user and continue
-		OutErrorMessage = LOCTEXT("AssetNameTooLong", "This asset name is too long. Please choose a shorter name.");
+		OutErrorMessage = FText::Format(LOCTEXT("SaveLayoutNameTooLong", "This layout file name is too long ({0} characters), the maximum is {1}. Please choose a shorter name. File name: {2}"),
+			FText::AsNumber(FilePath.Len()), FText::AsNumber(NAME_SIZE - 1), FText::FromString(FilePath));
 		// Return false to indicate that the user should enter a new name
 		return false;
 	}
 
 	// Make sure we are not creating an path that is too long for the OS
 	const FString FullPath = FPaths::ConvertRelativePathToFull(FilePath);
-	if (FilePath.Len() > FPlatformMisc::GetMaxPathLength())
+	if (FullPath.Len() >= FPlatformMisc::GetMaxPathLength())
 	{
 		// The full path for the asset is too long
-		OutErrorMessage = FText::Format(LOCTEXT("ObjectPathTooLong",
-			"The object path for the asset is too long, the maximum is '{0}'. \nPlease choose a shorter name for the asset or create it in a shallower folder structure."),
-			FText::AsNumber(FPlatformMisc::GetMaxPathLength()));
+		OutErrorMessage = FText::Format(LOCTEXT("SaveLayoutFullPathTooLong",
+			"The full path for the asset is too long ({0} characters), the maximum is {1}. Please choose a shorter name or create it in a shallower folder structure. Full path: {2}"),
+			FText::AsNumber(FullPath.Len()), FText::AsNumber(FPlatformMisc::GetMaxPathLength() - 1), FText::FromString(FullPath));
 		// Return false to indicate that the user should enter a new name
 		return false;
 	}
@@ -458,7 +459,7 @@ bool SSaveLayoutDialog::IsValidFilePathForCreation(const FString& FilePath, cons
 		if (FPaths::FileExists(FilePath))
 		{
 			// This asset already exists at this location, inform the user and continue
-			OutErrorMessage = FText::Format(LOCTEXT("RenameAssetAlreadyExists", "An asset already exists at this location with the name '{0}'."), FText::FromString(FileName));
+			OutErrorMessage = FText::Format(LOCTEXT("SaveLayoutAlreadyExists", "An asset already exists at this location with the name '{0}'."), FText::FromString(FileName));
 
 			// Return false to indicate that the user should enter a new name
 			return false;
