Copied CL# 10289192

Bug fixed: The "Save Layout As..." button (in "Window" --> "Save Layout" --> "Save Layout As...") was not localizing the text of the Layout Description field, only the one from the Layout Name field. It should localize both texts automatically. Fixed.

diff --git a/Engine/Source/Editor/MainFrame/Private/Menus/LayoutsMenu.cpp b/Engine/Source/Editor/MainFrame/Private/Menus/LayoutsMenu.cpp
--- a/Engine/Source/Editor/MainFrame/Private/Menus/LayoutsMenu.cpp
+++ b/Engine/Source/Editor/MainFrame/Private/Menus/LayoutsMenu.cpp
@@ -770,22 +770,23 @@ void SaveExportLayoutCommon(const FString& InDefaultDirectory, const bool bMustB
 	{
 		bool bWereFilesSelected = false;
 		TArray<FString> LayoutFilePaths;
-		TArray<FText> LayoutNames;
-		LayoutNames.Emplace(FLayoutSaveRestore::LoadSectionFromConfig(GEditorLayoutIni, "LayoutName"));
-		if (LayoutNames[0].ToString().Len() > 0)
-		{
-			LayoutNames[0] = FText::FromString(TEXT("Copy of ") + LayoutNames[0].ToString());
-		}
-		TArray<FText> LayoutDescriptions;
-		LayoutDescriptions.Emplace(FLayoutSaveRestore::LoadSectionFromConfig(GEditorLayoutIni, "LayoutDescription"));
-		if (LayoutDescriptions[0].ToString().Len() > 0)
-		{
-			LayoutDescriptions[0] = FText::FromString(TEXT("Copy of ") + LayoutDescriptions[0].ToString());
-		}
-		// "Save Layout As..."
 		bool bWasDialogOpened = bMustBeSavedInDefaultDirectory;
+		// "Save Layout As..."
 		if (bMustBeSavedInDefaultDirectory)
 		{
+			TArray<FText> LayoutNames;
+			LayoutNames.Emplace(FLayoutSaveRestore::LoadSectionFromConfig(GEditorLayoutIni, "LayoutName"));
+			TArray<FText> LayoutDescriptions;
+			LayoutDescriptions.Emplace(FLayoutSaveRestore::LoadSectionFromConfig(GEditorLayoutIni, "LayoutDescription"));
+			// We want to avoid the duplication of the file/name/description fields, so we add "Copy of " at the beginning of the name and description fields
+			if (LayoutNames[0].ToString().Len() > 0)
+			{
+				LayoutNames[0] = FText::FromString(TEXT("Copy of ") + LayoutNames[0].ToString());
+			}
+			if (LayoutDescriptions[0].ToString().Len() > 0)
+			{
+				LayoutDescriptions[0] = FText::FromString(TEXT("Copy of ") + LayoutDescriptions[0].ToString());
+			}
 			// Create SWidget for saving the layout in its own SWindow and block the thread until it is finished
 			const TSharedRef<FSaveLayoutDialogParams> SaveLayoutDialogParams = MakeShared<FSaveLayoutDialogParams>(InDefaultDirectory, TEXT(".ini"), LayoutNames, LayoutDescriptions);
 			bWasDialogOpened = FSaveLayoutDialogUtils::CreateSaveLayoutAsDialogInStandaloneWindow(SaveLayoutDialogParams);
@@ -797,14 +798,18 @@ void SaveExportLayoutCommon(const FString& InDefaultDirectory, const bool bMustB
 			// Update GEditorLayoutIni file if LayoutNames or LayoutDescriptions were modified by the user
 			if (bWasDialogOpened && LayoutNames.Num() > 0 && LayoutDescriptions.Num() > 0 && (LayoutNames[0].ToString().Len() > 0 || LayoutDescriptions[0].ToString().Len() > 0))
 			{
-				const FText& LayoutNameAsTextText = GenerateLocalizedTextForFile(LayoutNames[0]);
-				const FText& LayoutDescriptionAsTextText = GenerateLocalizedTextForFile(LayoutDescriptions[0]);
-				// Update fields
-				FLayoutSaveRestore::SaveSectionToConfig(GEditorLayoutIni, "LayoutName", LayoutNameAsTextText);
-				FLayoutSaveRestore::SaveSectionToConfig(GEditorLayoutIni, "LayoutDescription", LayoutDescriptions[0]);
-				// Flush file
-				const bool bRead = true;
-				GConfig->Flush(bRead, GEditorLayoutIni);
+				checkf(LayoutNames.Num() == LayoutDescriptions.Num(), TEXT("There should be the same number of LayoutNames and LayoutDescriptions."));
+				for (int32 Index = 0; Index < LayoutNames.Num(); ++Index)
+				{
+					const FText& LayoutNameAsTextText = GenerateLocalizedTextForFile(LayoutNames[Index]);
+					const FText& LayoutDescriptionAsTextText = GenerateLocalizedTextForFile(LayoutDescriptions[Index]);
+					// Update fields
+					FLayoutSaveRestore::SaveSectionToConfig(GEditorLayoutIni, "LayoutName", LayoutNameAsTextText);
+					FLayoutSaveRestore::SaveSectionToConfig(GEditorLayoutIni, "LayoutDescription", LayoutDescriptionAsTextText);
+					// Flush file
+					const bool bRead = true;
+					GConfig->Flush(bRead, GEditorLayoutIni);
+				}
 			}
 		}
 		// "Export Layout..." (or "Save Layout As..." dialog could not be opened)
