Disable editor fullscreen. It has to many issues overall that needs addressing

diff --git a/Engine/Source/Editor/MainFrame/Private/Frame/MainFrameActions.cpp b/Engine/Source/Editor/MainFrame/Private/Frame/MainFrameActions.cpp
--- a/Engine/Source/Editor/MainFrame/Private/Frame/MainFrameActions.cpp
+++ b/Engine/Source/Editor/MainFrame/Private/Frame/MainFrameActions.cpp
@@ -201,7 +201,7 @@ void FMainFrameCommands::RegisterCommands()
 	UI_COMMAND( SaveLayout, "Save Layout", "Save the layout customizations", EUserInterfaceActionType::Button, FInputChord() );
 	ActionList->MapAction( SaveLayout, FExecuteAction::CreateStatic( &FMainFrameActionCallbacks::SaveLayout) );
 
-#if !PLATFORM_MAC // Fullscreen mode in the editor is currently unsupported on Mac
+#if !PLATFORM_MAC && !PLATFORM_LINUX // Fullscreen mode in the editor is currently unsupported on Mac and Linux
 	UI_COMMAND( ToggleFullscreen, "Enable Fullscreen", "Enables fullscreen mode for the application, expanding across the entire monitor", EUserInterfaceActionType::ToggleButton, FInputChord(EModifierKey::Shift, EKeys::F11) );
 	ActionList->MapAction( ToggleFullscreen,
 		FExecuteAction::CreateStatic( &FMainFrameActionCallbacks::ToggleFullscreen_Execute ),
@@ -942,7 +942,7 @@ void FMainFrameActionCallbacks::SaveLayout()
 
 void FMainFrameActionCallbacks::ToggleFullscreen_Execute()
 {
-#if !PLATFORM_MAC // Fullscreen mode in the editor is currently unsupported on Mac
+#if !PLATFORM_MAC && !PLATFORM_LINUX // Fullscreen mode in the editor is currently unsupported on Mac or Linux
 	if ( GIsEditor && FApp::HasProjectName() )
 	{
 		static TWeakPtr<SDockTab> LevelEditorTabPtr = FGlobalTabmanager::Get()->InvokeTab(FTabId("LevelEditor"));
