Converted serveral menus to use new menu system: level editor, asset editor
#rnx

diff --git a/Engine/Source/Editor/LevelEditor/LevelEditor.Build.cs b/Engine/Source/Editor/LevelEditor/LevelEditor.Build.cs
--- a/Engine/Source/Editor/LevelEditor/LevelEditor.Build.cs
+++ b/Engine/Source/Editor/LevelEditor/LevelEditor.Build.cs
@@ -66,7 +66,8 @@ public class LevelEditor : ModuleRules
 				"HierarchicalLODUtilities",
 				"MaterialShaderQualitySettings",
 				"PixelInspectorModule",
-				"CommonMenuExtensions"
+				"CommonMenuExtensions",
+				"EditorMenus",
 			}
 		);
 
diff --git a/Engine/Source/Editor/LevelEditor/Private/LevelEditor.cpp b/Engine/Source/Editor/LevelEditor/Private/LevelEditor.cpp
--- a/Engine/Source/Editor/LevelEditor/Private/LevelEditor.cpp
+++ b/Engine/Source/Editor/LevelEditor/Private/LevelEditor.cpp
@@ -587,6 +587,19 @@ TSharedRef<IViewportLayoutEntity> FLevelEditorModule::FactoryViewport(FName InTy
 	return MakeShareable(new FLevelViewportLayoutEntity(ConstructionArgs));
 }
 
+TSharedPtr<FExtender> FLevelEditorModule::AssembleExtenders(TSharedRef<FUICommandList>& InCommandList, TArray<FLevelEditorMenuExtender>& MenuExtenderDelegates) const
+{
+	TArray<TSharedPtr<FExtender>> Extenders;
+	for (int32 i = 0; i < MenuExtenderDelegates.Num(); ++i)
+	{
+		if (MenuExtenderDelegates[i].IsBound())
+		{
+			Extenders.Add(MenuExtenderDelegates[i].Execute(InCommandList));
+		}
+	}
+	return FExtender::Combine(Extenders);
+}
+
 void FLevelEditorModule::BindGlobalLevelEditorCommands()
 {
 	check( !GlobalLevelEditorActions.IsValid() );
diff --git a/Engine/Source/Editor/LevelEditor/Private/LevelEditorMenu.cpp b/Engine/Source/Editor/LevelEditor/Private/LevelEditorMenu.cpp
--- a/Engine/Source/Editor/LevelEditor/Private/LevelEditorMenu.cpp
+++ b/Engine/Source/Editor/LevelEditor/Private/LevelEditorMenu.cpp
@@ -15,53 +15,57 @@
 #include "MRUFavoritesList.h"
 #include "Framework/Commands/GenericCommands.h"
 #include "IDocumentation.h"
+#include "EditorMenuSubsystem.h"
 
 #define LOCTEXT_NAMESPACE "LevelEditorMenu"
 
-TSharedRef< SWidget > FLevelEditorMenu::MakeLevelEditorMenu( const TSharedPtr<FUICommandList>& CommandList, TSharedPtr<class SLevelEditor> LevelEditor )
+void FLevelEditorMenu::RegisterLevelEditorMenus()
 {
 	struct Local
 	{
-		static void FillFileLoadAndSaveItems( FMenuBuilder& MenuBuilder )
+		static void RegisterFileLoadAndSaveItems()
 		{
+			UEditorMenu* Menu = UEditorMenuSubsystem::Get()->ExtendMenu("LevelEditor.MainMenu.File");
+			FEditorMenuSection& Section = Menu->FindOrAddSection("FileLoadAndSave");
+			FEditorMenuInsert InsertPos(NAME_None, EEditorMenuInsertType::First);
+
 			// New Level
-			MenuBuilder.AddMenuEntry( FLevelEditorCommands::Get().NewLevel );
+			Section.AddMenuEntry( FLevelEditorCommands::Get().NewLevel ).InsertPosition = InsertPos;
 
 			// Open Level
-			MenuBuilder.AddMenuEntry( FLevelEditorCommands::Get().OpenLevel );
+			Section.AddMenuEntry( FLevelEditorCommands::Get().OpenLevel ).InsertPosition = InsertPos;
 
 			// Open Asset
-			//@TODO: Doesn't work when summoned from here: MenuBuilder.AddMenuEntry( FGlobalEditorCommonCommands::Get().SummonOpenAssetDialog );
+			//@TODO: Doesn't work when summoned from here: Section.AddMenuEntry( FGlobalEditorCommonCommands::Get().SummonOpenAssetDialog );
 
 			// Save
-			MenuBuilder.AddMenuEntry( FLevelEditorCommands::Get().Save );
+			Section.AddMenuEntry( FLevelEditorCommands::Get().Save ).InsertPosition = InsertPos;
 	
 			// Save As
-			MenuBuilder.AddMenuEntry( FLevelEditorCommands::Get().SaveAs );
+			Section.AddMenuEntry( FLevelEditorCommands::Get().SaveAs ).InsertPosition = InsertPos;
 
 			// Save Levels
-			MenuBuilder.AddMenuEntry( FLevelEditorCommands::Get().SaveAllLevels );
+			Section.AddMenuEntry( FLevelEditorCommands::Get().SaveAllLevels ).InsertPosition = InsertPos;
 		}
 
-		static void FillFileRecentAndFavoriteFileItems( FMenuBuilder& MenuBuilder )
+		static void FillFileRecentAndFavoriteFileItems()
 		{
-			IMainFrameModule& MainFrameModule = FModuleManager::LoadModuleChecked<IMainFrameModule>( "MainFrame" );
-			const FMainMRUFavoritesList& RecentsAndFavorites = *MainFrameModule.GetMRUFavoritesList();
+			UEditorMenu* Menu = UEditorMenuSubsystem::Get()->ExtendMenu("LevelEditor.MainMenu.File");
+			FEditorMenuInsert SectionInsertPos("FileRecentFiles", EEditorMenuInsertType::Before);
 
 			// Import/Export
 			{
-				MenuBuilder.BeginSection("FileActors", LOCTEXT("ImportExportHeading", "Actors") );
+				FEditorMenuSection& Section = Menu->AddSection("FileActors", LOCTEXT("ImportExportHeading", "Actors"), SectionInsertPos);
 				{
 					// Import Into Level
-					MenuBuilder.AddMenuEntry(FLevelEditorCommands::Get().ImportScene);
+					Section.AddMenuEntry(FLevelEditorCommands::Get().ImportScene);
 
 					// Export All
-					MenuBuilder.AddMenuEntry( FLevelEditorCommands::Get().ExportAll );
+					Section.AddMenuEntry( FLevelEditorCommands::Get().ExportAll );
 
 					// Export Selected
-					MenuBuilder.AddMenuEntry( FLevelEditorCommands::Get().ExportSelected );
+					Section.AddMenuEntry( FLevelEditorCommands::Get().ExportSelected );
 				}
-				MenuBuilder.EndSection();
 			}
 
 
@@ -83,19 +87,18 @@ TSharedRef< SWidget > FLevelEditorMenu::MakeLevelEditorMenu( const TSharedPtr<FU
 						}
 					};
 
-					static void MakeFavoriteLevelMenu(FMenuBuilder& InMenuBuilder)
+					static void MakeFavoriteLevelMenu(UEditorMenu* InMenu)
 					{
 						// Add a button to add/remove the currently loaded map as a favorite
 						if (FLevelEditorActionCallbacks::ToggleFavorite_CanExecute())
 						{
-							InMenuBuilder.BeginSection("LevelEditorToggleFavorite");
+							FEditorMenuSection& Section = InMenu->AddSection("LevelEditorToggleFavorite");
 							{
 								TAttribute<FText> ToggleFavoriteLabel;
 								ToggleFavoriteLabel.BindStatic(&Local::GetToggleFavoriteLabelText);
-								InMenuBuilder.AddMenuEntry(FLevelEditorCommands::Get().ToggleFavorite, NAME_None, ToggleFavoriteLabel);
+								Section.AddMenuEntry(FLevelEditorCommands::Get().ToggleFavorite, ToggleFavoriteLabel);
 							}
-							InMenuBuilder.EndSection();
-							InMenuBuilder.AddMenuSeparator();
+							Section.AddMenuSeparator("LevelEditorToggleFavorite");
 						}
 						const FMainMRUFavoritesList& MRUFavorites = *FModuleManager::LoadModuleChecked<IMainFrameModule>("MainFrame").GetMRUFavoritesList();
 						const int32 NumFavorites = MRUFavorites.GetNumFavorites();
@@ -109,31 +112,37 @@ TSharedRef< SWidget > FLevelEditorMenu::MakeLevelEditorMenu( const TSharedPtr<FU
 							const FText ToolTip = FText::Format(LOCTEXT("FavoriteLevelToolTip", "Opens favorite level: {0}"), FText::FromString(CurFavorite));
 							const FText Label = FText::FromString(FPaths::GetBaseFilename(CurFavorite));
 
-							InMenuBuilder.AddMenuEntry(OpenFavoriteFile, NAME_None, Label, ToolTip);
+							InMenu->FindOrAddSection("Favorite").AddMenuEntry(OpenFavoriteFile, Label, ToolTip);
 						}
 					}
 				};
 
-				const int32 NumRecents = RecentsAndFavorites.GetNumItems();
-
-				MenuBuilder.BeginSection("FileFavoriteLevels");
+				FEditorMenuSection& Section = Menu->AddSection("FileFavoriteLevels", TAttribute<FText>(), SectionInsertPos);
 
-				if (NumRecents > 0)
+				Section.AddDynamicEntry("FileFavoriteLevels", FNewEditorMenuSectionDelegate::CreateLambda([](FEditorMenuSection& InSection)
 				{
-					MenuBuilder.AddSubMenu(
-						LOCTEXT("FavoriteLevelsSubMenu", "Favorite Levels"),
-						LOCTEXT("RecentLevelsSubMenu_ToolTip", "Select a level to load"),
-						FNewMenuDelegate::CreateStatic(&FFavoriteLevelMenu::MakeFavoriteLevelMenu), false, FSlateIcon(FEditorStyle::GetStyleSetName(), "MainFrame.FavoriteLevels"));
-				}
-
-				MenuBuilder.EndSection();
+					IMainFrameModule& MainFrameModule = FModuleManager::LoadModuleChecked<IMainFrameModule>("MainFrame");
+					const FMainMRUFavoritesList& RecentsAndFavorites = *MainFrameModule.GetMRUFavoritesList();
+					if (RecentsAndFavorites.GetNumItems() > 0)
+					{
+						InSection.AddEntry(FEditorMenuEntry::InitSubMenu(
+							"LevelEditor.MainMenu.File",
+							"FavoriteLevelsSubMenu",
+							LOCTEXT("FavoriteLevelsSubMenu", "Favorite Levels"),
+							LOCTEXT("RecentLevelsSubMenu_ToolTip", "Select a level to load"),
+							FNewEditorMenuDelegate::CreateStatic(&FFavoriteLevelMenu::MakeFavoriteLevelMenu),
+							false,
+							FSlateIcon(FEditorStyle::GetStyleSetName(), "MainFrame.FavoriteLevels")
+						));
+					}
+				}));
 			}
 
 			// Recent files
 			{
 				struct FRecentLevelMenu
 				{
-					static void MakeRecentLevelMenu( FMenuBuilder& InMenuBuilder )
+					static void MakeRecentLevelMenu( UEditorMenu* InMenu )
 					{
 						const FMainMRUFavoritesList& MRUFavorites = *FModuleManager::LoadModuleChecked<IMainFrameModule>( "MainFrame" ).GetMRUFavoritesList();
 						const int32 NumRecents = MRUFavorites.GetNumItems();
@@ -148,100 +157,93 @@ TSharedRef< SWidget > FLevelEditorMenu::MakeLevelEditorMenu( const TSharedPtr<FU
 							const FText ToolTip = FText::Format( LOCTEXT( "RecentLevelToolTip", "Opens recent level: {0}" ), FText::FromString( CurRecent ) );
 							const FText Label = FText::FromString( FPaths::GetBaseFilename( CurRecent ) );
 
-							InMenuBuilder.AddMenuEntry( OpenRecentFile, NAME_None, Label, ToolTip );
+							InMenu->FindOrAddSection("Recent").AddMenuEntry( OpenRecentFile, Label, ToolTip );
 						}
 					}
 				};
 
-				const int32 NumRecents = RecentsAndFavorites.GetNumItems();
-
-				MenuBuilder.BeginSection( "FileRecentLevels" );
-
-				if ( NumRecents > 0 )
+				FEditorMenuSection& Section = Menu->AddSection("FileRecentLevels", TAttribute<FText>(), SectionInsertPos);
+				Section.AddDynamicEntry("FileRecentLevels", FNewEditorMenuSectionDelegate::CreateLambda([](FEditorMenuSection& InSection)
 				{
-					MenuBuilder.AddSubMenu(
-						LOCTEXT("RecentLevelsSubMenu", "Recent Levels"),
-						LOCTEXT("RecentLevelsSubMenu_ToolTip", "Select a level to load"),
-						FNewMenuDelegate::CreateStatic( &FRecentLevelMenu::MakeRecentLevelMenu ), false, FSlateIcon(FEditorStyle::GetStyleSetName(), "MainFrame.RecentLevels") );
-				}
-
-				MenuBuilder.EndSection();
+					IMainFrameModule& MainFrameModule = FModuleManager::LoadModuleChecked<IMainFrameModule>( "MainFrame" );
+					const FMainMRUFavoritesList& RecentsAndFavorites = *MainFrameModule.GetMRUFavoritesList();
+					if (RecentsAndFavorites.GetNumItems() > 0)
+					{
+						InSection.AddEntry(FEditorMenuEntry::InitSubMenu(
+							"LevelEditor.MainMenu.File",
+							"RecentLevelsSubMenu",
+							LOCTEXT("RecentLevelsSubMenu", "Recent Levels"),
+							LOCTEXT("RecentLevelsSubMenu_ToolTip", "Select a level to load"),
+							FNewEditorMenuDelegate::CreateStatic(&FRecentLevelMenu::MakeRecentLevelMenu),
+							false,
+							FSlateIcon(FEditorStyle::GetStyleSetName(), "MainFrame.RecentLevels")
+						));
+					}
+				}));
 			}
 		}
 
-		static void FillEditMenu( FMenuBuilder& MenuBuilder )
+		static void ExtendEditMenu()
 		{
+			UEditorMenu* Menu = UEditorMenuSubsystem::Get()->RegisterMenu("LevelEditor.MainMenu.Edit", "MainFrame.MainMenu.Edit");
 			// Edit Actor
 			{
-				MenuBuilder.BeginSection("EditMain", LOCTEXT("MainHeading", "Edit") );
+				FEditorMenuSection& Section = Menu->AddSection("EditMain", LOCTEXT("MainHeading", "Edit"), FEditorMenuInsert("EditHistory", EEditorMenuInsertType::After));
 				{		
-					MenuBuilder.AddMenuEntry( FGenericCommands::Get().Cut );
-					MenuBuilder.AddMenuEntry( FGenericCommands::Get().Copy );
-					MenuBuilder.AddMenuEntry( FGenericCommands::Get().Paste );
+					Section.AddMenuEntry( FGenericCommands::Get().Cut );
+					Section.AddMenuEntry( FGenericCommands::Get().Copy );
+					Section.AddMenuEntry( FGenericCommands::Get().Paste );
 
-					MenuBuilder.AddMenuEntry( FGenericCommands::Get().Duplicate );
-					MenuBuilder.AddMenuEntry( FGenericCommands::Get().Delete );
+					Section.AddMenuEntry( FGenericCommands::Get().Duplicate );
+					Section.AddMenuEntry( FGenericCommands::Get().Delete );
 				}
-				MenuBuilder.EndSection();
 			}
 		}
 
-		static void ExtendHelpMenu( FMenuBuilder& MenuBuilder )
+		static void ExtendHelpMenu()
 		{
-			MenuBuilder.BeginSection("HelpBrowse", NSLOCTEXT("MainHelpMenu", "Browse", "Browse"));
+			UEditorMenu* Menu = UEditorMenuSubsystem::Get()->RegisterMenu("LevelEditor.MainMenu.Help", "MainFrame.MainMenu.Help");
+			FEditorMenuSection& Section = Menu->AddSection("HelpBrowse", NSLOCTEXT("MainHelpMenu", "Browse", "Browse"), FEditorMenuInsert("BugReporting", EEditorMenuInsertType::Before));
 			{
-				MenuBuilder.AddMenuEntry( FLevelEditorCommands::Get().BrowseDocumentation );
+				Section.AddMenuEntry( FLevelEditorCommands::Get().BrowseDocumentation );
 
-				MenuBuilder.AddMenuEntry( FLevelEditorCommands::Get().BrowseAPIReference );
+				Section.AddMenuEntry( FLevelEditorCommands::Get().BrowseAPIReference );
 
-				MenuBuilder.AddMenuEntry( FLevelEditorCommands::Get().BrowseCVars );
+				Section.AddMenuEntry( FLevelEditorCommands::Get().BrowseCVars );
 
-				MenuBuilder.AddMenuSeparator();
+				Section.AddMenuSeparator( "HelpBrowse" );
 
-				MenuBuilder.AddMenuEntry( FLevelEditorCommands::Get().BrowseViewportControls );
+				Section.AddMenuEntry( FLevelEditorCommands::Get().BrowseViewportControls );
 			}
-			MenuBuilder.EndSection();
-
-			
 		}
 	};
 
-	TSharedRef< FExtender > Extender( new FExtender() );
-	
+	UEditorMenuSubsystem* EditorMenus = UEditorMenuSubsystem::Get();
+	EditorMenus->RegisterMenu("LevelEditor.MainMenu", "MainFrame.MainMenu", EMultiBoxType::MenuBar);
+	EditorMenus->RegisterMenu("LevelEditor.MainMenu.File", "MainFrame.MainTabMenu.File");
+	EditorMenus->RegisterMenu("LevelEditor.MainMenu.Window", "MainFrame.MainMenu.Window");
+
 	// Add level loading and saving menu items
-	Extender->AddMenuExtension(
-		"FileLoadAndSave",
-		EExtensionHook::First,
-		CommandList.ToSharedRef(),
-		FMenuExtensionDelegate::CreateStatic( &Local::FillFileLoadAndSaveItems ) );
+	Local::RegisterFileLoadAndSaveItems();
 
 	// Add recent / favorites
-	Extender->AddMenuExtension(
-		"FileRecentFiles",
-		EExtensionHook::Before,
-		CommandList.ToSharedRef(),
-		FMenuExtensionDelegate::CreateStatic( &Local::FillFileRecentAndFavoriteFileItems ) );
+	Local::FillFileRecentAndFavoriteFileItems();
 
 	// Extend the Edit menu
-	Extender->AddMenuExtension(
-		"EditHistory",
-		EExtensionHook::After,
-		CommandList.ToSharedRef(),
-		FMenuExtensionDelegate::CreateStatic( &Local::FillEditMenu ) );
+	Local::ExtendEditMenu();
 
 	// Extend the Help menu
-	Extender->AddMenuExtension(
-		"BugReporting",
-		EExtensionHook::Before,
-		CommandList.ToSharedRef(),
-		FMenuExtensionDelegate::CreateStatic( &Local::ExtendHelpMenu ) );
+	Local::ExtendHelpMenu();
+}
 
+TSharedRef< SWidget > FLevelEditorMenu::MakeLevelEditorMenu( const TSharedPtr<FUICommandList>& CommandList, TSharedPtr<class SLevelEditor> LevelEditor )
+{
 	FLevelEditorModule& LevelEditorModule = FModuleManager::LoadModuleChecked<FLevelEditorModule>("LevelEditor");
-	LevelEditorModule.GetMenuExtensibilityManager()->AddExtender( Extender );
 	TSharedPtr<FExtender> Extenders = LevelEditorModule.GetMenuExtensibilityManager()->GetAllExtenders();
+	FEditorMenuContext EditorMenuContext(CommandList, Extenders.ToSharedRef());
 
 	IMainFrameModule& MainFrameModule = FModuleManager::LoadModuleChecked<IMainFrameModule>( "MainFrame" );
-	TSharedRef< SWidget > MenuBarWidget = MainFrameModule.MakeMainTabMenu( LevelEditor->GetTabManager(), Extenders.ToSharedRef() );
+	TSharedRef< SWidget > MenuBarWidget = MainFrameModule.MakeMainTabMenu( LevelEditor->GetTabManager(), "LevelEditor.MainMenu", EditorMenuContext );
 
 	return MenuBarWidget;
 }
@@ -261,4 +263,4 @@ TSharedRef< SWidget > FLevelEditorMenu::MakeNotificationBar( const TSharedPtr<FU
 	return NotificationBarBuilder.MakeWidget();
 }
 
-#undef LOCTEXT_NAMESPACE
+#undef LOCTEXT_NAMESPACE
\ No newline at end of file
diff --git a/Engine/Source/Editor/LevelEditor/Private/LevelEditorMenu.h b/Engine/Source/Editor/LevelEditor/Private/LevelEditorMenu.h
--- a/Engine/Source/Editor/LevelEditor/Private/LevelEditorMenu.h
+++ b/Engine/Source/Editor/LevelEditor/Private/LevelEditorMenu.h
@@ -16,6 +16,8 @@ class FLevelEditorMenu
 
 public:
 
+	static void RegisterLevelEditorMenus();
+
 	/**
 	 * Static: Creates a widget for the level editor's menu
 	 *
diff --git a/Engine/Source/Editor/LevelEditor/Private/LevelEditorToolBar.cpp b/Engine/Source/Editor/LevelEditor/Private/LevelEditorToolBar.cpp
--- a/Engine/Source/Editor/LevelEditor/Private/LevelEditorToolBar.cpp
+++ b/Engine/Source/Editor/LevelEditor/Private/LevelEditorToolBar.cpp
@@ -58,6 +58,8 @@
 #include "Misc/ScopedSlowTask.h"
 #include "MaterialShaderQualitySettings.h"
 #include "RHIShaderPlatformDefinitions.inl"
+#include "LevelEditorMenuContext.h"
+#include "EditorMenuSubsystem.h"
 
 namespace LevelEditorActionHelpers
 {
@@ -1146,26 +1148,27 @@ void LevelEditorActionHelpers::OnSelectPlayerControllerClassPicked(UClass* InCho
 	FSlateApplication::Get().DismissAllMenus();
 }
 
-
-/**
- * Static: Creates a widget for the level editor tool bar
- *
- * @return	New widget
- */
-TSharedRef< SWidget > FLevelEditorToolBar::MakeLevelEditorToolBar( const TSharedRef<FUICommandList>& InCommandList, const TSharedRef<SLevelEditor> InLevelEditor )
+void FLevelEditorToolBar::RegisterLevelEditorToolBar( const TSharedRef<FUICommandList>& InCommandList, const TSharedRef<SLevelEditor> InLevelEditor)
 {
-#define LOCTEXT_NAMESPACE "LevelEditorToolBar"
+	RegisterSourceControlMenu();
+	RegisterCinematicsMenu();
 
-	FLevelEditorModule& LevelEditorModule = FModuleManager::LoadModuleChecked<FLevelEditorModule>("LevelEditor");
-	TSharedPtr<FExtender> Extenders = LevelEditorModule.GetToolBarExtensibilityManager()->GetAllExtenders();
+#define LOCTEXT_NAMESPACE "LevelEditorToolBar"
 
-	static const FName LevelEditorToolBarName = "LevelEditorToolBar";
-	FToolBarBuilder ToolbarBuilder( InCommandList, FMultiBoxCustomization::AllowCustomization( LevelEditorToolBarName ), Extenders );
+	UEditorMenu* ToolbarBuilder = UEditorMenuSubsystem::Get()->RegisterMenu("LevelEditor.LevelEditorToolBar", NAME_None, EMultiBoxType::ToolBar);
 
-	ToolbarBuilder.BeginSection("File");
 	{
+		FEditorMenuSection& Section = ToolbarBuilder->AddSection("File");
+
 		// Save All Levels
-		ToolbarBuilder.AddToolBarButton( FLevelEditorCommands::Get().Save, NAME_None, TAttribute<FText>(), TAttribute<FText>(), FSlateIcon(FEditorStyle::GetStyleSetName(), "AssetEditor.SaveAsset") );
+		Section.AddEntry(FEditorMenuEntry::InitToolBarButton(
+			FLevelEditorCommands::Get().Save,
+			TAttribute<FText>(),
+			TAttribute<FText>(),
+			FSlateIcon(FEditorStyle::GetStyleSetName(), "AssetEditor.SaveAsset"),
+			NAME_None,
+			"SaveAllLevels"
+		));
 
 		// Source control buttons
 		{
@@ -1239,31 +1242,31 @@ TSharedRef< SWidget > FLevelEditorToolBar::MakeLevelEditorToolBar( const TShared
 
 			FSourceControlStatus::CheckSourceControlStatus();
 
-			ToolbarBuilder.AddComboButton(
+			Section.AddEntry(FEditorMenuEntry::InitComboButton(
+				"SourceControl",
 				FUIAction(),
 				FOnGetContent::CreateStatic(&FLevelEditorToolBar::GenerateSourceControlMenu, InCommandList),
 				LOCTEXT("SourceControl_Label", "Source Control"),
 				TAttribute<FText>::Create(&FSourceControlStatus::GetSourceControlTooltip),
 				TAttribute<FSlateIcon>::Create(&FSourceControlStatus::GetSourceControlIcon),
 				false
-				);
+				));
 		}
 	}
-	ToolbarBuilder.EndSection();
 
-	ToolbarBuilder.BeginSection("Content");
 	{
-		ToolbarBuilder.AddToolBarButton( FLevelEditorCommands::Get().OpenContentBrowser, NAME_None, LOCTEXT( "ContentBrowser_Override", "Content" ), TAttribute<FText>(), TAttribute<FSlateIcon>(), "LevelToolbarContent" );
+		FEditorMenuSection& Section = ToolbarBuilder->AddSection("Content");
+		Section.AddEntry(FEditorMenuEntry::InitToolBarButton(FLevelEditorCommands::Get().OpenContentBrowser, LOCTEXT( "ContentBrowser_Override", "Content" ), TAttribute<FText>(), TAttribute<FSlateIcon>(), "LevelToolbarContent"));
 		if (FLauncherPlatformModule::Get()->CanOpenLauncher(true)) 
 		{
-			ToolbarBuilder.AddToolBarButton(FLevelEditorCommands::Get().OpenMarketplace, NAME_None, LOCTEXT("Marketplace_Override", "Marketplace"), TAttribute<FText>(), TAttribute<FSlateIcon>(), "LevelToolbarMarketplace");
+			Section.AddEntry(FEditorMenuEntry::InitToolBarButton(FLevelEditorCommands::Get().OpenMarketplace, LOCTEXT("Marketplace_Override", "Marketplace"), TAttribute<FText>(), TAttribute<FSlateIcon>(), "LevelToolbarMarketplace"));
 		}
 	}
-	ToolbarBuilder.EndSection();
 
-	ToolbarBuilder.BeginSection("Settings");
+	FEditorMenuSection& SettingsSection = ToolbarBuilder->AddSection("Settings");
 	{
-		ToolbarBuilder.AddComboButton(
+		SettingsSection.AddEntry(FEditorMenuEntry::InitComboButton(
+			"LevelToolbarQuickSettings",
 			FUIAction(),
 			FOnGetContent::CreateStatic(&FLevelEditorToolBar::GenerateQuickSettingsMenu, InCommandList),
 			LOCTEXT("QuickSettingsCombo", "Settings"),
@@ -1271,7 +1274,7 @@ TSharedRef< SWidget > FLevelEditorToolBar::MakeLevelEditorToolBar( const TShared
 			FSlateIcon(FEditorStyle::GetStyleSetName(), "LevelEditor.GameSettings"),
 			false,
 			"LevelToolbarQuickSettings"
-			);
+			));
 
 	}
 
@@ -1402,70 +1405,71 @@ TSharedRef< SWidget > FLevelEditorToolBar::MakeLevelEditorToolBar( const TShared
 			}
 		};
 
-		ToolbarBuilder.AddToolBarButton(
+		SettingsSection.AddEntry(FEditorMenuEntry::InitToolBarButton( 
 			FLevelEditorCommands::Get().ToggleFeatureLevelPreview,
-			NAME_None,
 			TAttribute<FText>::Create(&FPreviewModeFunctionality::GetPreviewModeText),
         	TAttribute<FText>::Create(&FPreviewModeFunctionality::GetPreviewModeTooltip),
         	TAttribute<FSlateIcon>::Create(&FPreviewModeFunctionality::GetPreviewModeIcon)
-			);
+			));
 	}
-	ToolbarBuilder.EndSection();
 
-	ToolbarBuilder.BeginSection( NAME_None );
 	{
-		ToolbarBuilder.AddComboButton(
+		FEditorMenuSection& Section = ToolbarBuilder->AddSection("Misc");
+		Section.AddEntry(FEditorMenuEntry::InitComboButton(
+			"OpenBlueprint",
 			FUIAction(),
 			FOnGetContent::CreateStatic( &FLevelEditorToolBar::GenerateOpenBlueprintMenuContent, InCommandList, TWeakPtr<SLevelEditor>( InLevelEditor ) ),
 			LOCTEXT( "OpenBlueprint_Label", "Blueprints" ),
 			LOCTEXT( "OpenBlueprint_ToolTip", "List of world Blueprints available to the user for editing or creation." ),
 			FSlateIcon(FEditorStyle::GetStyleSetName(), "LevelEditor.OpenLevelBlueprint")
-			);
+			));
 
-		ToolbarBuilder.AddComboButton(
+		Section.AddEntry(FEditorMenuEntry::InitComboButton(
+			"EditCinematics",
 			FUIAction(),
 			FOnGetContent::CreateStatic( &FLevelEditorToolBar::GenerateCinematicsMenuContent, InCommandList, TWeakPtr<SLevelEditor>( InLevelEditor ) ),
 			LOCTEXT( "EditCinematics_Label", "Cinematics" ),
 			LOCTEXT( "EditCinematics_Tooltip", "Displays a list of Matinee and Level Sequence objects to open in their respective editors"),
 			FSlateIcon(FEditorStyle::GetStyleSetName(), "LevelEditor.EditMatinee") 
-			);
+			));
 
-		ToolbarBuilder.AddToolBarButton( FLevelEditorCommands::Get().ToggleVR, NAME_None, LOCTEXT("ToggleVR", "VR Mode") );
+		Section.AddEntry(FEditorMenuEntry::InitToolBarButton( FLevelEditorCommands::Get().ToggleVR, LOCTEXT("ToggleVR", "VR Mode")) );
 	}
-	ToolbarBuilder.EndSection();
-	
-	ToolbarBuilder.BeginSection("Compile");
+
 	{
+		FEditorMenuSection& Section = ToolbarBuilder->AddSection("Compile");
 		// Build			
-		ToolbarBuilder.AddToolBarButton( FLevelEditorCommands::Get().Build, NAME_None, LOCTEXT("BuildAll", "Build") );
+		Section.AddEntry(FEditorMenuEntry::InitToolBarButton( FLevelEditorCommands::Get().Build, LOCTEXT("BuildAll", "Build")) );
 
 		// Build menu drop down
-		ToolbarBuilder.AddComboButton(
+		Section.AddEntry(FEditorMenuEntry::InitComboButton(
+			"BuildComboButton",
 			FUIAction(),
 			FOnGetContent::CreateStatic( &FLevelEditorToolBar::GenerateBuildMenuContent, InCommandList ),
 			LOCTEXT( "BuildCombo_Label", "Build Options" ),
 			LOCTEXT( "BuildComboToolTip", "Build options menu" ),
 			FSlateIcon(FEditorStyle::GetStyleSetName(), "LevelEditor.Build"),
-			true);
+			true));
 
 		// Only show the compile options on machines with the solution (assuming they can build it)
 		if ( FSourceCodeNavigation::IsCompilerAvailable() )
 		{
 			// Since we can always add new code to the project, only hide these buttons if we haven't done so yet
-			ToolbarBuilder.AddToolBarButton(
+			Section.AddEntry(FEditorMenuEntry::InitToolBarButton(
+				"CompileButton",
 				FUIAction(
 					FExecuteAction::CreateStatic(&FLevelEditorActionCallbacks::RecompileGameCode_Clicked),
 					FCanExecuteAction::CreateStatic(&FLevelEditorActionCallbacks::Recompile_CanExecute),
 					FIsActionChecked(),
 					FIsActionButtonVisible::CreateStatic(FLevelEditorActionCallbacks::CanShowSourceCodeActions)),
-				NAME_None,
 				LOCTEXT( "CompileMenuButton", "Compile" ),
 				FLevelEditorCommands::Get().RecompileGameCode->GetDescription(),
 				FSlateIcon(FEditorStyle::GetStyleSetName(), "LevelEditor.Recompile")
-				);
+				));
 
 #if WITH_LIVE_CODING
-			ToolbarBuilder.AddComboButton(
+			Section.AddEntry(FEditorMenuEntry::InitComboButton(
+				"CompileComboButton",
 				FUIAction(
 					FExecuteAction(),
 					FCanExecuteAction(),
@@ -1476,19 +1480,34 @@ TSharedRef< SWidget > FLevelEditorToolBar::MakeLevelEditorToolBar( const TShared
 				LOCTEXT( "CompileComboToolTip", "Compile options menu" ),
 				FSlateIcon(FEditorStyle::GetStyleSetName(), "LevelEditor.Recompile"),
 				true
-				);
+				));
 #endif
 		}
 	}
-	ToolbarBuilder.EndSection();
 	
-	ToolbarBuilder.BeginSection("Game");
+	ToolbarBuilder->AddDynamicSection("Game", FNewToolBarDelegateLegacy::CreateLambda([](FToolBarBuilder& InBuilder, UEditorMenu* InData)
 	{
 		// Add the shared play-world commands that will be shown on the Kismet toolbar as well
-		FPlayWorldCommands::BuildToolbar(ToolbarBuilder, true);
-	}
-	ToolbarBuilder.EndSection();
+		FPlayWorldCommands::BuildToolbar(InBuilder, true);
+	}));
+
+#undef LOCTEXT_NAMESPACE
+}
+
+/**
+ * Static: Creates a widget for the level editor tool bar
+ *
+ * @return	New widget
+ */
+TSharedRef< SWidget > FLevelEditorToolBar::MakeLevelEditorToolBar( const TSharedRef<FUICommandList>& InCommandList, const TSharedRef<SLevelEditor> InLevelEditor )
+{
+#define LOCTEXT_NAMESPACE "LevelEditorToolBar"
 
+	FLevelEditorModule& LevelEditorModule = FModuleManager::LoadModuleChecked<FLevelEditorModule>("LevelEditor");
+	FEditorMenuContext MenuContext(InCommandList, LevelEditorModule.GetToolBarExtensibilityManager()->GetAllExtenders());
+	ULevelEditorMenuContext* LevelEditorMenuContext = NewObject<ULevelEditorMenuContext>();
+	LevelEditorMenuContext->SlateLevelEditor = InLevelEditor;
+	MenuContext.AddObject(LevelEditorMenuContext);
 
 	// Create the tool bar!
 	return
@@ -1497,7 +1516,7 @@ TSharedRef< SWidget > FLevelEditorToolBar::MakeLevelEditorToolBar( const TShared
 		.BorderImage( FEditorStyle::GetBrush("NoBorder") )
 		.IsEnabled( FSlateApplication::Get().GetNormalExecutionAttribute() )
 		[
-			ToolbarBuilder.MakeWidget()
+			UEditorMenuSubsystem::Get()->GenerateWidget("LevelEditor.LevelEditorToolBar", MenuContext)
 		];
 #undef LOCTEXT_NAMESPACE
 }
@@ -1510,17 +1529,7 @@ TSharedRef< SWidget > FLevelEditorToolBar::GenerateBuildMenuContent( TSharedRef<
 
 	// Get all menu extenders for this context menu from the level editor module
 	FLevelEditorModule& LevelEditorModule = FModuleManager::GetModuleChecked<FLevelEditorModule>( TEXT("LevelEditor") );
-	TArray<FLevelEditorModule::FLevelEditorMenuExtender> MenuExtenderDelegates = LevelEditorModule.GetAllLevelEditorToolbarBuildMenuExtenders();
-
-	TArray<TSharedPtr<FExtender>> Extenders;
-	for (int32 i = 0; i < MenuExtenderDelegates.Num(); ++i)
-	{
-		if (MenuExtenderDelegates[i].IsBound())
-		{
-			Extenders.Add(MenuExtenderDelegates[i].Execute(InCommandList));
-		}
-	}
-	TSharedPtr<FExtender> MenuExtender = FExtender::Combine(Extenders);
+	TSharedPtr<FExtender> MenuExtender = LevelEditorModule.AssembleExtenders(InCommandList, LevelEditorModule.GetAllLevelEditorToolbarBuildMenuExtenders());
 
 	const bool bShouldCloseWindowAfterMenuSelection = true;
 	FMenuBuilder MenuBuilder( bShouldCloseWindowAfterMenuSelection, InCommandList, MenuExtender );
@@ -1929,17 +1938,7 @@ TSharedRef< SWidget > FLevelEditorToolBar::GenerateQuickSettingsMenu( TSharedRef
 
 	// Get all menu extenders for this context menu from the level editor module
 	FLevelEditorModule& LevelEditorModule = FModuleManager::GetModuleChecked<FLevelEditorModule>( TEXT("LevelEditor") );
-	TArray<FLevelEditorModule::FLevelEditorMenuExtender> MenuExtenderDelegates = LevelEditorModule.GetAllLevelEditorToolbarViewMenuExtenders();
-
-	TArray<TSharedPtr<FExtender>> Extenders;
-	for (int32 i = 0; i < MenuExtenderDelegates.Num(); ++i)
-	{
-		if (MenuExtenderDelegates[i].IsBound())
-		{
-			Extenders.Add(MenuExtenderDelegates[i].Execute(InCommandList));
-		}
-	}
-	TSharedPtr<FExtender> MenuExtender = FExtender::Combine(Extenders);
+	TSharedPtr<FExtender> MenuExtender = LevelEditorModule.AssembleExtenders(InCommandList, LevelEditorModule.GetAllLevelEditorToolbarViewMenuExtenders());
 
 	const bool bShouldCloseWindowAfterMenuSelection = true;
 	FMenuBuilder MenuBuilder( bShouldCloseWindowAfterMenuSelection, InCommandList, MenuExtender );
@@ -2059,73 +2058,63 @@ TSharedRef< SWidget > FLevelEditorToolBar::GenerateQuickSettingsMenu( TSharedRef
 
 TSharedRef< SWidget > FLevelEditorToolBar::GenerateSourceControlMenu(TSharedRef<FUICommandList> InCommandList)
 {
-#define LOCTEXT_NAMESPACE "LevelToolBarSourceControlMenu"
-
 	// Get all menu extenders for this context menu from the level editor module
 	FLevelEditorModule& LevelEditorModule = FModuleManager::GetModuleChecked<FLevelEditorModule>(TEXT("LevelEditor"));
-	TArray<FLevelEditorModule::FLevelEditorMenuExtender> MenuExtenderDelegates = LevelEditorModule.GetAllLevelEditorToolbarSourceControlMenuExtenders();
+	TSharedPtr<FExtender> MenuExtender = LevelEditorModule.AssembleExtenders(InCommandList, LevelEditorModule.GetAllLevelEditorToolbarSourceControlMenuExtenders());
 
-	TArray<TSharedPtr<FExtender>> Extenders;
-	for (int32 i = 0; i < MenuExtenderDelegates.Num(); ++i)
-	{
-		if (MenuExtenderDelegates[i].IsBound())
-		{
-			Extenders.Add(MenuExtenderDelegates[i].Execute(InCommandList));
-		}
-	}
-	TSharedPtr<FExtender> MenuExtender = FExtender::Combine(Extenders);
+	FEditorMenuContext Context(InCommandList, MenuExtender);
+	return UEditorMenuSubsystem::Get()->GenerateWidget("LevelEditor.LevelEditorToolBar.SourceControl", Context);
+}
 
-	const bool bShouldCloseWindowAfterMenuSelection = true;
-	FMenuBuilder MenuBuilder(bShouldCloseWindowAfterMenuSelection, InCommandList, MenuExtender);
+void FLevelEditorToolBar::RegisterSourceControlMenu()
+{
+#define LOCTEXT_NAMESPACE "LevelToolBarSourceControlMenu"
 
-	MenuBuilder.BeginSection("SourceControlActions", LOCTEXT("SourceControlMenuHeadingActions", "Actions"));
+	UEditorMenu* Menu = UEditorMenuSubsystem::Get()->RegisterMenu("LevelEditor.LevelEditorToolBar.SourceControl");
+	Menu->bShouldCloseWindowAfterMenuSelection = true;
+	FEditorMenuSection& Section = Menu->AddSection("SourceControlActions", LOCTEXT("SourceControlMenuHeadingActions", "Actions"));
 
-	ISourceControlModule& SourceControlModule = ISourceControlModule::Get();
-	if (ISourceControlModule::Get().IsEnabled() && ISourceControlModule::Get().GetProvider().IsAvailable())
-	{
-		MenuBuilder.AddMenuEntry(
-			FLevelEditorCommands::Get().ChangeSourceControlSettings, 
-			NAME_None,
-			TAttribute<FText>(),
-			TAttribute<FText>(),
-			FSlateIcon(FEditorStyle::GetStyleSetName(), "SourceControl.Actions.ChangeSettings")
-			);
-	}
-	else
+	Section.AddDynamicEntry("ConnectToSourceControl", FNewEditorMenuSectionDelegate::CreateLambda([](FEditorMenuSection& InSection)
 	{
-		MenuBuilder.AddMenuEntry(
-			FLevelEditorCommands::Get().ConnectToSourceControl,
-			NAME_None,
-			TAttribute<FText>(),
-			TAttribute<FText>(),
-			FSlateIcon(FEditorStyle::GetStyleSetName(), "SourceControl.Actions.Connect")
-			);
-	}
+		ISourceControlModule& SourceControlModule = ISourceControlModule::Get();
+		if (ISourceControlModule::Get().IsEnabled() && ISourceControlModule::Get().GetProvider().IsAvailable())
+		{
+			InSection.AddMenuEntry(
+				FLevelEditorCommands::Get().ChangeSourceControlSettings, 
+				TAttribute<FText>(),
+				TAttribute<FText>(),
+				FSlateIcon(FEditorStyle::GetStyleSetName(), "SourceControl.Actions.ChangeSettings")
+				);
+		}
+		else
+		{
+			InSection.AddMenuEntry(
+				FLevelEditorCommands::Get().ConnectToSourceControl,
+				TAttribute<FText>(),
+				TAttribute<FText>(),
+				FSlateIcon(FEditorStyle::GetStyleSetName(), "SourceControl.Actions.Connect")
+				);
+		}
+	}));
 
-	MenuBuilder.AddMenuSeparator("SourceControlConnectionSeparator");
+	Section.AddMenuSeparator("SourceControlConnectionSeparator");
 
-	MenuBuilder.AddMenuEntry(
+	Section.AddMenuEntry(
 		FLevelEditorCommands::Get().CheckOutModifiedFiles,
-		NAME_None,
 		TAttribute<FText>(),
 		TAttribute<FText>(),
 		FSlateIcon(FEditorStyle::GetStyleSetName(), "SourceControl.Actions.CheckOut")
 		);
 
 
-	MenuBuilder.AddMenuEntry(
+	Section.AddMenuEntry(
 		FLevelEditorCommands::Get().SubmitToSourceControl,
-		NAME_None,
 		TAttribute<FText>(),
 		TAttribute<FText>(),
 		FSlateIcon(FEditorStyle::GetStyleSetName(), "SourceControl.Actions.Submit")
 		);
 
-	MenuBuilder.EndSection();
-
 #undef LOCTEXT_NAMESPACE
-
-	return MenuBuilder.MakeWidget();
 }
 
 TSharedRef< SWidget > FLevelEditorToolBar::GenerateOpenBlueprintMenuContent( TSharedRef<FUICommandList> InCommandList, TWeakPtr< SLevelEditor > InLevelEditor )
@@ -2290,102 +2279,83 @@ void FLevelEditorToolBar::OnOpenSubLevelBlueprint( ULevel* InLevel )
 
 TSharedRef< SWidget > FLevelEditorToolBar::GenerateCinematicsMenuContent( TSharedRef<FUICommandList> InCommandList, TWeakPtr<SLevelEditor> LevelEditorWeakPtr )
 {
-#define LOCTEXT_NAMESPACE "LevelToolBarCinematicsMenu"
-
 	FLevelEditorModule& LevelEditorModule = FModuleManager::LoadModuleChecked<FLevelEditorModule>("LevelEditor");
-	TSharedPtr<FExtender> Extender = FExtender::Combine(LevelEditorModule.GetAllLevelEditorToolbarCinematicsMenuExtenders());
+	FEditorMenuContext Context(InCommandList, FExtender::Combine(LevelEditorModule.GetAllLevelEditorToolbarCinematicsMenuExtenders()));
+	ULevelEditorMenuContext* LevelEditorMenuContext = NewObject<ULevelEditorMenuContext>();
+	LevelEditorMenuContext->SlateLevelEditor = LevelEditorWeakPtr;
+	Context.AddObject(LevelEditorMenuContext);
+	return UEditorMenuSubsystem::Get()->GenerateWidget("LevelEditor.LevelEditorToolBar.Cinematics", Context);
+}
 
-	const bool bShouldCloseWindowAfterMenuSelection = true;
-	FMenuBuilder MenuBuilder( bShouldCloseWindowAfterMenuSelection, InCommandList, Extender );
+void FLevelEditorToolBar::RegisterCinematicsMenu()
+{
+#define LOCTEXT_NAMESPACE "LevelToolBarCinematicsMenu"
+
+	UEditorMenu* Menu = UEditorMenuSubsystem::Get()->RegisterMenu("LevelEditor.LevelEditorToolBar.Cinematics");
+	Menu->bShouldCloseWindowAfterMenuSelection = true;
 
-	using namespace SceneOutliner;
+	{
+		FEditorMenuSection& Section = Menu->AddSection("LevelEditorNewMatinee", LOCTEXT("MatineeMenuCombo_NewHeading", "New"));
+		Section.AddMenuEntry( FLevelEditorCommands::Get().AddMatinee, TAttribute<FText>(), TAttribute<FText>(), FSlateIcon(FEditorStyle::GetStyleSetName(), "ClassIcon.MatineeActor") );
+	}
 
-	// We can't build a list of Matinees and LevelSequenceActors while the current World is a PIE world.
-	FInitializationOptions InitOptions;
+	//Add a heading to separate the existing cinematics from the 'Add New Cinematic Actor' button
+	FEditorMenuSection& ExistingCinematicSection = Menu->AddSection("LevelEditorExistingCinematic", LOCTEXT("CinematicMenuCombo_ExistingHeading", "Edit Existing Cinematic"));
+	ExistingCinematicSection.AddDynamicEntry("LevelEditorExistingCinematic", FNewEditorMenuSectionDelegate::CreateLambda([](FEditorMenuSection& InSection)
 	{
-		InitOptions.Mode = ESceneOutlinerMode::ActorPicker;
+		ULevelEditorMenuContext* FoundContext = InSection.Context.Find<ULevelEditorMenuContext>();
+		if (!FoundContext)
+		{
+			return;
+		}
 
-		// We hide the header row to keep the UI compact.
-		// @todo: Might be useful to have this sometimes, actually.  Ideally the user could summon it.
-		InitOptions.bShowHeaderRow = false;
-		InitOptions.bShowSearchBox = false;
-		InitOptions.bShowCreateNewFolder = false;
+		UWorld* World = FoundContext->SlateLevelEditor.IsValid() ? FoundContext->SlateLevelEditor.Pin()->GetWorld() : nullptr;
+		const bool bHasAnyCinematicsActors = !!TActorIterator<AMatineeActor>(World) || !!TActorIterator<ALevelSequenceActor>(World);
+		if (!bHasAnyCinematicsActors)
+		{
+			return;
+		}
 
-		InitOptions.ColumnMap.Add(FBuiltInColumnTypes::Label(), FColumnInfo(EColumnVisibility::Visible, 0));
-		InitOptions.ColumnMap.Add(FBuiltInColumnTypes::ActorInfo(), FColumnInfo(EColumnVisibility::Visible, 10));
+		using namespace SceneOutliner;
 
-		// Only display Matinee and MovieScene actors
-		auto ActorFilter = [](const AActor* Actor){
-			return Actor->IsA( AMatineeActor::StaticClass() ) || Actor->IsA( ALevelSequenceActor::StaticClass() );
-		};
-		InitOptions.Filters->AddFilterPredicate( FActorFilterPredicate::CreateLambda( ActorFilter ) );
-	}
+		// We can't build a list of Matinees and LevelSequenceActors while the current World is a PIE world.
+		FInitializationOptions InitOptions;
+		{
+			InitOptions.Mode = ESceneOutlinerMode::ActorPicker;
 
-	// actor selector to allow the user to choose an actor
-	FSceneOutlinerModule& SceneOutlinerModule = FModuleManager::LoadModuleChecked<FSceneOutlinerModule>( "SceneOutliner" );
-	TSharedRef< SWidget > MiniSceneOutliner =
-		SNew( SVerticalBox )
-		+SVerticalBox::Slot()
-		.AutoHeight()
-		.MaxHeight(400.0f)
-		[
-			SceneOutlinerModule.CreateSceneOutliner(
-				InitOptions,
-				FOnActorPicked::CreateStatic( &FLevelEditorToolBar::OnCinematicsActorPicked ) )
-		];
+			// We hide the header row to keep the UI compact.
+			// @todo: Might be useful to have this sometimes, actually.  Ideally the user could summon it.
+			InitOptions.bShowHeaderRow = false;
+			InitOptions.bShowSearchBox = false;
+			InitOptions.bShowCreateNewFolder = false;
 
-	static const FName DefaultForegroundName("DefaultForeground");
+			InitOptions.ColumnMap.Add(FBuiltInColumnTypes::Label(), FColumnInfo(EColumnVisibility::Visible, 0));
+			InitOptions.ColumnMap.Add(FBuiltInColumnTypes::ActorInfo(), FColumnInfo(EColumnVisibility::Visible, 10));
 
-	// Give the scene outliner a border and background
-	const FSlateBrush* BackgroundBrush = FEditorStyle::GetBrush( "Menu.Background" );
-	TSharedRef< SBorder > RootBorder =
-		SNew( SBorder )
-		.Padding(3)
-		.BorderImage( BackgroundBrush )
-		.ForegroundColor( FEditorStyle::GetSlateColor(DefaultForegroundName) )
+			// Only display Matinee and MovieScene actors
+			auto ActorFilter = [](const AActor* Actor) {
+				return Actor->IsA(AMatineeActor::StaticClass()) || Actor->IsA(ALevelSequenceActor::StaticClass());
+			};
+			InitOptions.Filters->AddFilterPredicate(FActorFilterPredicate::CreateLambda(ActorFilter));
+		}
 
-		// Assign the box panel as the child
-		[
-			SNew( SVerticalBox )
-			+SVerticalBox::Slot()
+		// actor selector to allow the user to choose an actor
+		FSceneOutlinerModule& SceneOutlinerModule = FModuleManager::LoadModuleChecked<FSceneOutlinerModule>("SceneOutliner");
+		TSharedRef< SWidget > MiniSceneOutliner =
+			SNew(SVerticalBox)
+			+ SVerticalBox::Slot()
 			.AutoHeight()
-			.Padding( 5 )
-			.HAlign( HAlign_Center )
+			.MaxHeight(400.0f)
 			[
-				SNew( STextBlock )
-				.Text( LOCTEXT( "SelectCinematicsActorToEdit", "Select an actor" ) )
-			]
-
-			+SVerticalBox::Slot()
-				.AutoHeight()
-				.Padding( 2 )
-				[
-					MiniSceneOutliner
-				]
-		]
-	;
-
-	MenuBuilder.BeginSection("LevelEditorNewMatinee", LOCTEXT("MatineeMenuCombo_NewHeading", "New"));
-	{
-		MenuBuilder.AddMenuEntry( FLevelEditorCommands::Get().AddMatinee, NAME_None, TAttribute<FText>(), TAttribute<FText>(), FSlateIcon(FEditorStyle::GetStyleSetName(), "ClassIcon.MatineeActor") );
-	}
-	MenuBuilder.EndSection();
+				SceneOutlinerModule.CreateSceneOutliner(
+					InitOptions,
+					FOnActorPicked::CreateStatic(&FLevelEditorToolBar::OnCinematicsActorPicked))
+			];
 
-	UWorld* World = LevelEditorWeakPtr.Pin()->GetWorld();
-	const bool bHasAnyCinematicsActors = !!TActorIterator<AMatineeActor>(World) || !!TActorIterator<ALevelSequenceActor>(World);
+		InSection.AddEntry(FEditorMenuEntry::InitWidget("LevelEditorExistingCinematic", MiniSceneOutliner, FText::GetEmpty(), true));
+	}));
 
-	//Add a heading to separate the existing cinematics from the 'Add New Cinematic Actor' button
-	MenuBuilder.BeginSection("LevelEditorExistingCinematic", LOCTEXT( "CinematicMenuCombo_ExistingHeading", "Edit Existing Cinematic" ) );
-	{
-		if( bHasAnyCinematicsActors )
-		{
-			MenuBuilder.AddWidget(MiniSceneOutliner, FText::GetEmpty(), true);
-		}
-	}
-	MenuBuilder.EndSection();
 #undef LOCTEXT_NAMESPACE
-
-	return MenuBuilder.MakeWidget();
 }
 
 void FLevelEditorToolBar::OnCinematicsActorPicked( AActor* Actor )
diff --git a/Engine/Source/Editor/LevelEditor/Private/LevelEditorToolBar.h b/Engine/Source/Editor/LevelEditor/Private/LevelEditorToolBar.h
--- a/Engine/Source/Editor/LevelEditor/Private/LevelEditorToolBar.h
+++ b/Engine/Source/Editor/LevelEditor/Private/LevelEditorToolBar.h
@@ -8,6 +8,8 @@
 #include "Framework/Commands/UICommandList.h"
 #include "Editor/LevelEditor/Private/SLevelEditor.h"
 
+class UEditorMenu;
+
 /**
  * Unreal level editor main toolbar
  */
@@ -23,6 +25,7 @@ public:
 	 */
 	static TSharedRef< SWidget > MakeLevelEditorToolBar( const TSharedRef<FUICommandList>& InCommandList, const TSharedRef<SLevelEditor> InLevelEditor );
 
+	static void RegisterLevelEditorToolBar( const TSharedRef<FUICommandList>& InCommandList, const TSharedRef<SLevelEditor> InLevelEditor );
 
 protected:
 
@@ -89,4 +92,9 @@ protected:
 	 * @param InLevel	The level to open the Blueprint of (creates if needed)
 	 */
 	static void OnOpenSubLevelBlueprint( ULevel* InLevel );
+
+private:
+
+	static void RegisterSourceControlMenu();
+	static void RegisterCinematicsMenu();
 };
diff --git a/Engine/Source/Editor/LevelEditor/Private/SLevelEditor.cpp b/Engine/Source/Editor/LevelEditor/Private/SLevelEditor.cpp
--- a/Engine/Source/Editor/LevelEditor/Private/SLevelEditor.cpp
+++ b/Engine/Source/Editor/LevelEditor/Private/SLevelEditor.cpp
@@ -142,6 +142,12 @@ void SLevelEditor::BindCommands()
 		);
 }
 
+void SLevelEditor::RegisterMenus()
+{
+	FLevelEditorMenu::RegisterLevelEditorMenus();
+	FLevelEditorToolBar::RegisterLevelEditorToolBar(LevelEditorCommands.ToSharedRef(), SharedThis(this));
+}
+
 void SLevelEditor::Construct( const SLevelEditor::FArguments& InArgs)
 {
 	// Important: We use raw bindings here because we are releasing our binding in our destructor (where a weak pointer would be invalid)
@@ -154,6 +160,8 @@ void SLevelEditor::Construct( const SLevelEditor::FArguments& InArgs)
 
 	BindCommands();
 
+	RegisterMenus();
+
 	// We need to register when modes list changes so that we can refresh the auto generated commands.
 	FEditorModeRegistry::Get().OnRegisteredModesChanged().AddRaw(this, &SLevelEditor::EditorModeCommandsChanged);
 
diff --git a/Engine/Source/Editor/LevelEditor/Private/SLevelEditor.h b/Engine/Source/Editor/LevelEditor/Private/SLevelEditor.h
--- a/Engine/Source/Editor/LevelEditor/Private/SLevelEditor.h
+++ b/Engine/Source/Editor/LevelEditor/Private/SLevelEditor.h
@@ -148,6 +148,9 @@ private:
 	 */
 	void BindCommands();
 
+	/** Registers menus associated with level editor */
+	void RegisterMenus();
+
 	/**
 	 * Fills the level editor with content, using the layout string, or the default if
 	 * no layout string is passed in
diff --git a/Engine/Source/Editor/LevelEditor/Private/SLevelViewport.cpp b/Engine/Source/Editor/LevelEditor/Private/SLevelViewport.cpp
--- a/Engine/Source/Editor/LevelEditor/Private/SLevelViewport.cpp
+++ b/Engine/Source/Editor/LevelEditor/Private/SLevelViewport.cpp
@@ -391,18 +391,8 @@ TSharedRef<SWidget> SLevelViewport::GenerateLevelMenu() const
 	FWorldBrowserModule& WorldBrowserModule = FModuleManager::LoadModuleChecked<FWorldBrowserModule>("WorldBrowser");
 	// Get all menu extenders for this context menu from the level editor module
 	FLevelEditorModule& LevelEditorModule = FModuleManager::GetModuleChecked<FLevelEditorModule>(TEXT("LevelEditor"));
-	TArray<FLevelEditorModule::FLevelEditorMenuExtender> MenuExtenderDelegates = LevelEditorModule.GetAllLevelEditorLevelMenuExtenders();
-	TArray<TSharedPtr<FExtender>> Extenders;
-
-	TSharedPtr<FUICommandList> InCommandList = GetCommandList();
-	for (int32 i = 0; i < MenuExtenderDelegates.Num(); ++i)
-	{
-		if (MenuExtenderDelegates[i].IsBound())
-		{
-			Extenders.Add(MenuExtenderDelegates[i].Execute(InCommandList.ToSharedRef()));
-		}
-	}
-	TSharedPtr<FExtender> MenuExtender = FExtender::Combine(Extenders);
+	TSharedRef<FUICommandList> InCommandList = GetCommandList().ToSharedRef();
+	TSharedPtr<FExtender> MenuExtender = LevelEditorModule.AssembleExtenders(InCommandList, LevelEditorModule.GetAllLevelEditorLevelMenuExtenders());
 
 	// Create the menu
 	const bool bShouldCloseWindowAfterMenuSelection = true;
diff --git a/Engine/Source/Editor/LevelEditor/Private/SLevelViewportToolBar.cpp b/Engine/Source/Editor/LevelEditor/Private/SLevelViewportToolBar.cpp
--- a/Engine/Source/Editor/LevelEditor/Private/SLevelViewportToolBar.cpp
+++ b/Engine/Source/Editor/LevelEditor/Private/SLevelViewportToolBar.cpp
@@ -509,20 +509,11 @@ TSharedRef<SWidget> SLevelViewportToolBar::GenerateOptionsMenu() const
 	Viewport.Pin()->OnFloatingButtonClicked();
 
 	const FLevelViewportCommands& LevelViewportActions = FLevelViewportCommands::Get();
+	TSharedRef<FUICommandList> CommandList = Viewport.Pin()->GetCommandList().ToSharedRef();
 
 	// Get all menu extenders for this context menu from the level editor module
 	FLevelEditorModule& LevelEditorModule = FModuleManager::GetModuleChecked<FLevelEditorModule>( TEXT("LevelEditor") );
-	TArray<FLevelEditorModule::FLevelEditorMenuExtender> MenuExtenderDelegates = LevelEditorModule.GetAllLevelViewportOptionsMenuExtenders();
-	
-	TArray<TSharedPtr<FExtender>> Extenders;
-	for (int32 i = 0; i < MenuExtenderDelegates.Num(); ++i)
-	{
-		if (MenuExtenderDelegates[i].IsBound())
-		{
-			Extenders.Add(MenuExtenderDelegates[i].Execute(Viewport.Pin()->GetCommandList().ToSharedRef()));
-		}
-	}
-	TSharedPtr<FExtender> MenuExtender = FExtender::Combine(Extenders);
+	TSharedPtr<FExtender> MenuExtender = LevelEditorModule.AssembleExtenders(CommandList, LevelEditorModule.GetAllLevelViewportOptionsMenuExtenders());
 
 	const bool bIsPerspective = Viewport.Pin()->GetLevelViewportClient().IsPerspective();
 	const bool bInShouldCloseWindowAfterMenuSelection = true;
@@ -940,20 +931,11 @@ TSharedRef<SWidget> SLevelViewportToolBar::GenerateShowMenu() const
 	Viewport.Pin()->OnFloatingButtonClicked();
 
 	const FLevelViewportCommands& Actions = FLevelViewportCommands::Get();
+	TSharedRef<FUICommandList> CommandList = Viewport.Pin()->GetCommandList().ToSharedRef();
 
 	// Get all menu extenders for this context menu from the level editor module
 	FLevelEditorModule& LevelEditorModule = FModuleManager::GetModuleChecked<FLevelEditorModule>(TEXT("LevelEditor"));
-	TArray<FLevelEditorModule::FLevelEditorMenuExtender> MenuExtenderDelegates = LevelEditorModule.GetAllLevelViewportShowMenuExtenders();
-
-	TArray<TSharedPtr<FExtender>> Extenders;
-	for (int32 i = 0; i < MenuExtenderDelegates.Num(); ++i)
-	{
-		if (MenuExtenderDelegates[i].IsBound())
-		{
-			Extenders.Add(MenuExtenderDelegates[i].Execute(Viewport.Pin()->GetCommandList().ToSharedRef()));
-		}
-	}
-	TSharedPtr<FExtender> MenuExtender = FExtender::Combine(Extenders);
+	TSharedPtr<FExtender> MenuExtender = LevelEditorModule.AssembleExtenders(CommandList, LevelEditorModule.GetAllLevelViewportShowMenuExtenders());
 
 	const bool bInShouldCloseWindowAfterMenuSelection = true;
 	FMenuBuilder ShowMenuBuilder( bInShouldCloseWindowAfterMenuSelection, Viewport.Pin()->GetCommandList(), MenuExtender);
diff --git a/Engine/Source/Editor/LevelEditor/Public/LevelEditor.h b/Engine/Source/Editor/LevelEditor/Public/LevelEditor.h
--- a/Engine/Source/Editor/LevelEditor/Public/LevelEditor.h
+++ b/Engine/Source/Editor/LevelEditor/Public/LevelEditor.h
@@ -241,6 +241,8 @@ public:
 	virtual TSharedPtr<FExtensibilityManager> GetModeBarExtensibilityManager() {return ModeBarExtensibilityManager;}
 	virtual TSharedPtr<FExtensibilityManager> GetNotificationBarExtensibilityManager() {return NotificationBarExtensibilityManager;}
 
+	virtual TSharedPtr<FExtender> AssembleExtenders(TSharedRef<FUICommandList>& InCommandList, TArray<FLevelEditorMenuExtender>& MenuExtenderDelegates) const;
+
 	DECLARE_EVENT_OneParam(ILevelEditor, FOnRegisterTabs, TSharedPtr<FTabManager>);
 	FOnRegisterTabs& OnRegisterTabs() { return RegisterTabs; }
 
diff --git a/Engine/Source/Editor/LocalizationDashboard/LocalizationDashboard.Build.cs b/Engine/Source/Editor/LocalizationDashboard/LocalizationDashboard.Build.cs
--- a/Engine/Source/Editor/LocalizationDashboard/LocalizationDashboard.Build.cs
+++ b/Engine/Source/Editor/LocalizationDashboard/LocalizationDashboard.Build.cs
@@ -34,6 +34,7 @@ public class LocalizationDashboard : ModuleRules
 				"LocalizationCommandletExecution",
 				"LocalizationService",
 				"InternationalizationSettings",
+				"EditorMenus",
 			}
 		);
 
diff --git a/Engine/Source/Editor/LocalizationDashboard/Private/LocalizationDashboard.cpp b/Engine/Source/Editor/LocalizationDashboard/Private/LocalizationDashboard.cpp
--- a/Engine/Source/Editor/LocalizationDashboard/Private/LocalizationDashboard.cpp
+++ b/Engine/Source/Editor/LocalizationDashboard/Private/LocalizationDashboard.cpp
@@ -20,6 +20,7 @@
 #include "LocalizationSettings.h"
 #include "LocalizationDashboardSettings.h"
 #include "SSettingsEditorCheckoutNotice.h"
+#include "EditorMenuSubsystem.h"
 
 #define LOCTEXT_NAMESPACE "LocalizationDashboard"
 
@@ -144,9 +145,9 @@ void SLocalizationDashboard::Construct(const FArguments& InArgs, const TSharedPt
 
 	Layout = FLayoutSaveRestore::LoadFromConfig(GEditorLayoutIni, Layout);
 
-	const TSharedRef<FExtender> MenuExtender = MakeShareable(new FExtender());
+	FEditorMenuContext EditorMenuContext;
 	IMainFrameModule& MainFrameModule = FModuleManager::LoadModuleChecked<IMainFrameModule>( "MainFrame" );
-	const TSharedRef<SWidget> MenuWidget = MainFrameModule.MakeMainMenu( TabManager, MenuExtender );
+	const TSharedRef<SWidget> MenuWidget = MainFrameModule.MakeMainMenu( TabManager, "LocalizationDashboard.MainMenu", EditorMenuContext );
 
 	FString ConfigFilePath;
 	ConfigFilePath = GetDefault<ULocalizationSettings>()->GetDefaultConfigFilename();
diff --git a/Engine/Source/Editor/MainFrame/MainFrame.Build.cs b/Engine/Source/Editor/MainFrame/MainFrame.Build.cs
--- a/Engine/Source/Editor/MainFrame/MainFrame.Build.cs
+++ b/Engine/Source/Editor/MainFrame/MainFrame.Build.cs
@@ -40,6 +40,7 @@ public class MainFrame : ModuleRules
 				"DeviceProfileEditor",
 				"UndoHistory",
 				"Analytics",
+				"EditorMenus",
 			}
 		);
 
diff --git a/Engine/Source/Editor/MainFrame/Private/MainFrameModule.cpp b/Engine/Source/Editor/MainFrame/Private/MainFrameModule.cpp
--- a/Engine/Source/Editor/MainFrame/Private/MainFrameModule.cpp
+++ b/Engine/Source/Editor/MainFrame/Private/MainFrameModule.cpp
@@ -234,15 +234,15 @@ void FMainFrameModule::CreateDefaultMainFrame( const bool bStartImmersive, const
 }
 
 
-TSharedRef<SWidget> FMainFrameModule::MakeMainMenu( const TSharedPtr<FTabManager>& TabManager, const TSharedRef< FExtender > Extender ) const
+TSharedRef<SWidget> FMainFrameModule::MakeMainMenu( const TSharedPtr<FTabManager>& TabManager, const FName MenuName, FEditorMenuContext& EditorMenuContext ) const
 {
-	return FMainMenu::MakeMainMenu( TabManager, Extender );
+	return FMainMenu::MakeMainMenu( TabManager, MenuName, EditorMenuContext );
 }
 
 
-TSharedRef<SWidget> FMainFrameModule::MakeMainTabMenu( const TSharedPtr<FTabManager>& TabManager, const TSharedRef< FExtender > Extender ) const
+TSharedRef<SWidget> FMainFrameModule::MakeMainTabMenu( const TSharedPtr<FTabManager>& TabManager, const FName MenuName, FEditorMenuContext& EditorMenuContext ) const
 {
-	return FMainMenu::MakeMainTabMenu( TabManager, Extender );
+	return FMainMenu::MakeMainTabMenu( TabManager, MenuName, EditorMenuContext );
 }
 
 
diff --git a/Engine/Source/Editor/MainFrame/Private/MainFrameModule.h b/Engine/Source/Editor/MainFrame/Private/MainFrameModule.h
--- a/Engine/Source/Editor/MainFrame/Private/MainFrameModule.h
+++ b/Engine/Source/Editor/MainFrame/Private/MainFrameModule.h
@@ -23,8 +23,8 @@ public:
 	// IMainFrameModule interface
 
 	virtual void CreateDefaultMainFrame( const bool bStartImmersive, const bool bStartPIE ) override;
-	virtual TSharedRef<SWidget> MakeMainMenu( const TSharedPtr<FTabManager>& TabManager, const TSharedRef< FExtender > Extender ) const override;
-	virtual TSharedRef<SWidget> MakeMainTabMenu( const TSharedPtr<FTabManager>& TabManager, const TSharedRef< FExtender > Extender ) const override;
+	virtual TSharedRef<SWidget> MakeMainMenu( const TSharedPtr<FTabManager>& TabManager, const FName MenuName, FEditorMenuContext& EditorMenuContext ) const override;
+	virtual TSharedRef<SWidget> MakeMainTabMenu( const TSharedPtr<FTabManager>& TabManager, const FName MenuName, FEditorMenuContext& EditorMenuContext ) const override;
 	virtual TSharedRef<SWidget> MakeDeveloperTools( const TArray<FMainFrameDeveloperTool>& AdditionalTools ) const override;
 
 	virtual bool IsWindowInitialized( ) const override
diff --git a/Engine/Source/Editor/MainFrame/Private/Menus/MainMenu.cpp b/Engine/Source/Editor/MainFrame/Private/Menus/MainMenu.cpp
--- a/Engine/Source/Editor/MainFrame/Private/Menus/MainMenu.cpp
+++ b/Engine/Source/Editor/MainFrame/Private/Menus/MainMenu.cpp
@@ -23,6 +23,8 @@
 #include "Menus/SettingsMenu.h"
 #include "Menus/MainFrameTranslationEditorMenu.h"
 
+#include "EditorMenuSubsystem.h"
+
 #include "WorkspaceMenuStructure.h"
 #include "WorkspaceMenuStructureModule.h"
 #include "Features/EditorFeatures.h"
@@ -35,30 +37,38 @@
 #define LOCTEXT_NAMESPACE "MainFileMenu"
 
 
-void FMainMenu::FillFileMenu( FMenuBuilder& MenuBuilder, const TSharedRef<FExtender> Extender )
+void FMainMenu::RegisterFileMenu()
 {
-	MenuBuilder.BeginSection("FileLoadAndSave", LOCTEXT("LoadSandSaveHeading", "Load and Save"));
+	UEditorMenuSubsystem* EditorMenus = UEditorMenuSubsystem::Get();
+	UEditorMenu* FileMenu = EditorMenus->RegisterMenu("MainFrame.MainMenu.File");
+
+	FEditorMenuSection& FileLoadAndSaveSection = FileMenu->AddSection("FileLoadAndSave", LOCTEXT("LoadSandSaveHeading", "Load and Save"));
 	{
 		// Open Asset...
-		MenuBuilder.AddMenuEntry( FGlobalEditorCommonCommands::Get().SummonOpenAssetDialog );
+		FileLoadAndSaveSection.AddMenuEntry(FGlobalEditorCommonCommands::Get().SummonOpenAssetDialog);
 
 		// Save All
-		MenuBuilder.AddMenuEntry(FMainFrameCommands::Get().SaveAll, "SaveAll");
+		FileLoadAndSaveSection.AddMenuEntry(FMainFrameCommands::Get().SaveAll);
 
 		// Choose specific files to save
-		MenuBuilder.AddMenuEntry(FMainFrameCommands::Get().ChooseFilesToSave, "ChooseFilesToSave");
+		FileLoadAndSaveSection.AddMenuEntry(FMainFrameCommands::Get().ChooseFilesToSave);
 
-		if (ISourceControlModule::Get().IsEnabled() && ISourceControlModule::Get().GetProvider().IsAvailable())
+		FileLoadAndSaveSection.AddDynamicEntry("SourceControl", FNewEditorMenuSectionDelegate::CreateLambda([](FEditorMenuSection& InSection)
 		{
-			// Choose specific files to submit
-			MenuBuilder.AddMenuEntry(FMainFrameCommands::Get().ChooseFilesToCheckIn, "ChooseFilesToCheckIn");
-		}
-		else
-		{
-			MenuBuilder.AddMenuEntry(FMainFrameCommands::Get().ConnectToSourceControl, "ConnectToSourceControl");
-		}
+			if (ISourceControlModule::Get().IsEnabled() && ISourceControlModule::Get().GetProvider().IsAvailable())
+			{
+				// Choose specific files to submit
+				InSection.AddMenuEntry(FMainFrameCommands::Get().ChooseFilesToCheckIn);
+			}
+			else
+			{
+				InSection.AddMenuEntry(FMainFrameCommands::Get().ConnectToSourceControl);
+			}
+		}));
 	}
-	MenuBuilder.EndSection();
+
+	RegisterFileProjectMenu();
+	RegisterRecentFileAndExitMenuItems();
 }
 
 #undef LOCTEXT_NAMESPACE
@@ -66,10 +76,11 @@ void FMainMenu::FillFileMenu( FMenuBuilder& MenuBuilder, const TSharedRef<FExten
 
 #define LOCTEXT_NAMESPACE "MainEditMenu"
 
-void FMainMenu::FillEditMenu( FMenuBuilder& MenuBuilder, const TSharedRef< FExtender > Extender, const TSharedPtr<FTabManager> TabManager )
+void FMainMenu::RegisterEditMenu()
 {
-	MenuBuilder.BeginSection("EditHistory", LOCTEXT("HistoryHeading", "History"));
+	UEditorMenu* EditMenu = UEditorMenuSubsystem::Get()->RegisterMenu("MainFrame.MainMenu.Edit");
 	{
+		FEditorMenuSection& Section = EditMenu->AddSection("EditHistory", LOCTEXT("HistoryHeading", "History"));
 		struct Local
 		{
 			/** @return Returns a dynamic text string for Undo that contains the name of the action */
@@ -88,25 +99,25 @@ void FMainMenu::FillEditMenu( FMenuBuilder& MenuBuilder, const TSharedRef< FExte
 		// Undo
 		TAttribute<FText> DynamicUndoLabel;
 		DynamicUndoLabel.BindStatic(&Local::GetUndoLabelText);
-		MenuBuilder.AddMenuEntry( FGenericCommands::Get().Undo, "Undo", DynamicUndoLabel); // TAttribute< FString >::Create( &Local::GetUndoLabelText ) );
+		Section.AddMenuEntry(FGenericCommands::Get().Undo, DynamicUndoLabel); // TAttribute< FString >::Create( &Local::GetUndoLabelText ) );
 
 		// Redo
 		TAttribute< FText > DynamicRedoLabel;
 		DynamicRedoLabel.BindStatic( &Local::GetRedoLabelText );
-		MenuBuilder.AddMenuEntry(FGenericCommands::Get().Redo, "Redo", DynamicRedoLabel); // TAttribute< FString >::Create( &Local::GetRedoLabelText ) );
+		Section.AddMenuEntry(FGenericCommands::Get().Redo, DynamicRedoLabel); // TAttribute< FString >::Create( &Local::GetRedoLabelText ) );
 
 		// Show undo history
-		MenuBuilder.AddMenuEntry(
+		Section.AddMenuEntry(
+			"UndoHistory",
 			LOCTEXT("UndoHistoryTabTitle", "Undo History"),
 			LOCTEXT("UndoHistoryTooltipText", "View the entire undo history."),
 			FSlateIcon(FEditorStyle::GetStyleSetName(), "UndoHistory.TabIcon"),
 			FUIAction(FExecuteAction::CreateStatic(&FUndoHistoryModule::ExecuteOpenUndoHistory))
 			);
 	}
-	MenuBuilder.EndSection();
 
-	MenuBuilder.BeginSection("EditLocalTabSpawners", LOCTEXT("ConfigurationHeading", "Configuration"));
 	{
+		FEditorMenuSection& Section = EditMenu->AddSection("EditLocalTabSpawners", LOCTEXT("ConfigurationHeading", "Configuration"));
 		if (GetDefault<UEditorExperimentalSettings>()->bToolbarCustomization)
 		{
 			FUIAction ToggleMultiBoxEditMode(
@@ -115,45 +126,56 @@ void FMainMenu::FillEditMenu( FMenuBuilder& MenuBuilder, const TSharedRef< FExte
 				FIsActionChecked::CreateStatic(&FMultiBoxSettings::IsInToolbarEditMode)
 			);
 		
-			MenuBuilder.AddMenuEntry(
+			Section.AddMenuEntry(
+				"EditToolbars",
 				LOCTEXT("EditToolbarsLabel", "Edit Toolbars"),
 				LOCTEXT("EditToolbarsToolTip", "Allows customization of each toolbar"),
 				FSlateIcon(),
 				ToggleMultiBoxEditMode,
-				NAME_None,
 				EUserInterfaceActionType::ToggleButton
 			);
 
-			// Automatically populate tab spawners from TabManager
-			if (TabManager.IsValid())
+			Section.AddDynamicEntry("TabManager", FNewEditorMenuDelegateLegacy::CreateLambda([](FMenuBuilder& MenuBuilder, UEditorMenu* MenuData)
 			{
-				const IWorkspaceMenuStructure& MenuStructure = WorkspaceMenu::GetMenuStructure();
-				TabManager->PopulateTabSpawnerMenu(MenuBuilder, MenuStructure.GetEditOptions());
-			}
+				if (USlateTabManagerContext* TabManagerContext = MenuData->FindContext<USlateTabManagerContext>())
+				{
+					if (TabManagerContext->TabManager.IsValid())
+					{
+						// Automatically populate tab spawners from TabManager
+						const IWorkspaceMenuStructure& MenuStructure = WorkspaceMenu::GetMenuStructure();
+						TabManagerContext->TabManager.Pin()->PopulateTabSpawnerMenu(MenuBuilder, MenuStructure.GetEditOptions());
+					}
+				}
+			}));
 		}
 
 		if (GetDefault<UEditorStyleSettings>()->bExpandConfigurationMenus)
 		{
-			MenuBuilder.AddSubMenu(
+			Section.AddEntry(FEditorMenuEntry::InitSubMenu(
+				EditMenu->GetMenuName(),
+				"EditorPreferencesSubMenu",
 				LOCTEXT("EditorPreferencesSubMenuLabel", "Editor Preferences"),
 				LOCTEXT("EditorPreferencesSubMenuToolTip", "Configure the behavior and features of this Editor"),
-				FNewMenuDelegate::CreateStatic(&FSettingsMenu::MakeMenu, FName("Editor")),
+				FNewEditorMenuDelegate::CreateStatic(&FSettingsMenu::MakeMenu, FName("Editor")),
 				false,
 				FSlateIcon(FEditorStyle::GetStyleSetName(), "EditorPreferences.TabIcon")
-			);
+			));
 
-			MenuBuilder.AddSubMenu(
+			Section.AddEntry(FEditorMenuEntry::InitSubMenu(
+				EditMenu->GetMenuName(),
+				"ProjectSettingsSubMenu",
 				LOCTEXT("ProjectSettingsSubMenuLabel", "Project Settings"),
 				LOCTEXT("ProjectSettingsSubMenuToolTip", "Change the settings of the currently loaded project"),
-				FNewMenuDelegate::CreateStatic(&FSettingsMenu::MakeMenu, FName("Project")),
+				FNewEditorMenuDelegate::CreateStatic(&FSettingsMenu::MakeMenu, FName("Project")),
 				false,
 				FSlateIcon(FEditorStyle::GetStyleSetName(), "ProjectSettings.TabIcon")
-			);
+			));
 		}
 		else
 		{
 #if !PLATFORM_MAC // Handled by app's menu in menu bar
-			MenuBuilder.AddMenuEntry(
+			Section.AddMenuEntry(
+				"EditorPreferencesMenu",
 				LOCTEXT("EditorPreferencesMenuLabel", "Editor Preferences..."),
 				LOCTEXT("EditorPreferencesMenuToolTip", "Configure the behavior and features of the Unreal Editor."),
 				FSlateIcon(FEditorStyle::GetStyleSetName(), "EditorPreferences.TabIcon"),
@@ -161,7 +183,8 @@ void FMainMenu::FillEditMenu( FMenuBuilder& MenuBuilder, const TSharedRef< FExte
 			);
 #endif
 
-			MenuBuilder.AddMenuEntry(
+			Section.AddMenuEntry(
+				"ProjectSettingsMenu",
 				LOCTEXT("ProjectSettingsMenuLabel", "Project Settings..."),
 				LOCTEXT("ProjectSettingsMenuToolTip", "Change the settings of the currently loaded project."),
 				FSlateIcon(FEditorStyle::GetStyleSetName(), "ProjectSettings.TabIcon"),
@@ -169,42 +192,53 @@ void FMainMenu::FillEditMenu( FMenuBuilder& MenuBuilder, const TSharedRef< FExte
 			);
 		}
 
-		//@todo The tab system needs to be able to be extendable by plugins [9/3/2013 Justin.Sargent]
-		if (IModularFeatures::Get().IsModularFeatureAvailable(EditorFeatures::PluginsEditor))
+		Section.AddDynamicEntry("PluginsEditor", FNewEditorMenuDelegateLegacy::CreateLambda([](FMenuBuilder& InBuilder, UEditorMenu* InData)
 		{
-			FGlobalTabmanager::Get()->PopulateTabSpawnerMenu(MenuBuilder, "PluginsEditor");
-		}
+			//@todo The tab system needs to be able to be extendable by plugins [9/3/2013 Justin.Sargent]
+			if (IModularFeatures::Get().IsModularFeatureAvailable(EditorFeatures::PluginsEditor))
+			{
+				FGlobalTabmanager::Get()->PopulateTabSpawnerMenu(InBuilder, "PluginsEditor");
+			}
+		}));
 	}
-	MenuBuilder.EndSection();
 }
 
 #undef LOCTEXT_NAMESPACE
 
 #define LOCTEXT_NAMESPACE "MainWindowMenu"
 
-void FMainMenu::FillWindowMenu( FMenuBuilder& MenuBuilder, const TSharedRef< FExtender > Extender, const TSharedPtr<FTabManager> TabManager )
+void FMainMenu::RegisterWindowMenu()
 {
+	UEditorMenu* Menu = UEditorMenuSubsystem::Get()->RegisterMenu("MainFrame.MainMenu.Window");
+
 	// Automatically populate tab spawners from TabManager
-	if (TabManager.IsValid())
+	Menu->AddDynamicSection("TabManagerSection", FNewEditorMenuDelegateLegacy::CreateLambda([](FMenuBuilder& InBuilder, UEditorMenu* InData)
 	{
-		// Local editor tabs
-		TabManager->PopulateLocalTabSpawnerMenu(MenuBuilder);
-	
-		// General tabs
-		const IWorkspaceMenuStructure& MenuStructure = WorkspaceMenu::GetMenuStructure();
-		TabManager->PopulateTabSpawnerMenu(MenuBuilder, MenuStructure.GetStructureRoot());
-	}
+		if (USlateTabManagerContext* TabManagerContext = InData->FindContext<USlateTabManagerContext>())
+		{
+			TSharedPtr<FTabManager> TabManager = TabManagerContext->TabManager.Pin();
+			if (TabManager.IsValid())
+			{
+				// Local editor tabs
+				TabManager->PopulateLocalTabSpawnerMenu(InBuilder);
+
+				// General tabs
+				const IWorkspaceMenuStructure& MenuStructure = WorkspaceMenu::GetMenuStructure();
+				TabManager->PopulateTabSpawnerMenu(InBuilder, MenuStructure.GetStructureRoot());
+			}
+		}
+	}));
 
-	MenuBuilder.BeginSection("WindowGlobalTabSpawners");
 	{
-		MenuBuilder.AddMenuEntry(
+		FEditorMenuSection& Section = Menu->AddSection("WindowGlobalTabSpawners");
+		Section.AddMenuEntry(
+			"ProjectLauncher",
 			LOCTEXT("ProjectLauncherLabel", "Project Launcher"),
 			LOCTEXT("ProjectLauncherToolTip", "The Project Launcher provides advanced workflows for packaging, deploying and launching your projects."),
 			FSlateIcon(FEditorStyle::GetStyleSetName(), "Launcher.TabIcon"),
 			FUIAction(FExecuteAction::CreateStatic(&FMainMenu::OpenProjectLauncher))
 			);
 	}
-	MenuBuilder.EndSection();
 
 	{
 		// This is a temporary home for the spawners of experimental features that must be explicitly enabled.
@@ -215,12 +249,13 @@ void FMainMenu::FillWindowMenu( FMenuBuilder& MenuBuilder, const TSharedRef< FEx
 		// Make sure at least one is enabled before creating the section
 		if (bLocalizationDashboard || bTranslationPicker)
 		{
-			MenuBuilder.BeginSection("ExperimentalTabSpawners", LOCTEXT("ExperimentalTabSpawnersHeading", "Experimental"));
+			FEditorMenuSection& Section = Menu->AddSection("ExperimentalTabSpawners", LOCTEXT("ExperimentalTabSpawnersHeading", "Experimental"));
 			{
 				// Localization Dashboard
 				if (bLocalizationDashboard)
 				{
-					MenuBuilder.AddMenuEntry(
+					Section.AddMenuEntry(
+						"LocalizationDashboard",
 						LOCTEXT("LocalizationDashboardLabel", "Localization Dashboard"),
 						LOCTEXT("LocalizationDashboardToolTip", "Open the Localization Dashboard for this Project."),
 						FSlateIcon(),
@@ -231,7 +266,8 @@ void FMainMenu::FillWindowMenu( FMenuBuilder& MenuBuilder, const TSharedRef< FEx
 				// Translation Picker
 				if (bTranslationPicker)
 				{
-					MenuBuilder.AddMenuEntry(
+					Section.AddMenuEntry(
+						"TranslationPicker",
 						LOCTEXT("TranslationPickerMenuItem", "Translation Picker"),
 						LOCTEXT("TranslationPickerMenuItemToolTip", "Launch the Translation Picker to Modify Editor Translations"),
 						FSlateIcon(),
@@ -239,115 +275,136 @@ void FMainMenu::FillWindowMenu( FMenuBuilder& MenuBuilder, const TSharedRef< FEx
 						);
 				}
 			}
-			MenuBuilder.EndSection();
 		}
 	}
 
-	MenuBuilder.BeginSection("WindowLayout", NSLOCTEXT("MainAppMenu", "LayoutManagementHeader", "Layout"));
 	{
-		MenuBuilder.AddMenuEntry(FMainFrameCommands::Get().ResetLayout);
-		MenuBuilder.AddMenuEntry(FMainFrameCommands::Get().SaveLayout);
+		FEditorMenuSection& Section = Menu->AddSection("WindowLayout", NSLOCTEXT("MainAppMenu", "LayoutManagementHeader", "Layout"));
+		Section.AddMenuEntry(FMainFrameCommands::Get().ResetLayout);
+		Section.AddMenuEntry(FMainFrameCommands::Get().SaveLayout);
 #if !PLATFORM_MAC && !PLATFORM_LINUX // On Mac/Linux windowed fullscreen mode in the editor is currently unavailable
-		MenuBuilder.AddMenuEntry(FMainFrameCommands::Get().ToggleFullscreen);
+		Section.AddMenuEntry(FMainFrameCommands::Get().ToggleFullscreen);
 #endif
 	}
-	MenuBuilder.EndSection();
 }
 
 #undef LOCTEXT_NAMESPACE
 
 
-void FMainMenu::FillHelpMenu( FMenuBuilder& MenuBuilder, const TSharedRef< FExtender > Extender )
+#define LOCTEXT_NAMESPACE "MainHelpMenu"
+
+void FMainMenu::RegisterHelpMenu()
 {
-	MenuBuilder.BeginSection("BugReporting", NSLOCTEXT("MainHelpMenu", "BugsReporting", "Bugs"));
+	UEditorMenu* Menu = UEditorMenuSubsystem::Get()->RegisterMenu("MainFrame.MainMenu.Help");
+	FEditorMenuSection& BugReportingSection = Menu->AddSection("BugReporting", NSLOCTEXT("MainHelpMenu", "BugsReporting", "Bugs"));
 	{
-		MenuBuilder.AddMenuEntry(FMainFrameCommands::Get().ReportABug);
-		MenuBuilder.AddMenuEntry(FMainFrameCommands::Get().OpenIssueTracker);
+		BugReportingSection.AddMenuEntry(FMainFrameCommands::Get().ReportABug);
+		BugReportingSection.AddMenuEntry(FMainFrameCommands::Get().OpenIssueTracker);
 	}
-	MenuBuilder.EndSection();
 
-	MenuBuilder.BeginSection("HelpOnline", NSLOCTEXT("MainHelpMenu", "Online", "Help Online"));
+	FEditorMenuSection& HelpOnlineSection = Menu->AddSection("HelpOnline", NSLOCTEXT("MainHelpMenu", "Online", "Help Online"));
 	{
-		MenuBuilder.AddMenuEntry(FMainFrameCommands::Get().VisitSupportWebSite);
-		MenuBuilder.AddMenuEntry(FMainFrameCommands::Get().VisitForums);
-		MenuBuilder.AddMenuEntry(FMainFrameCommands::Get().VisitSearchForAnswersPage);
-		MenuBuilder.AddMenuEntry(FMainFrameCommands::Get().VisitWiki);
+		HelpOnlineSection.AddMenuEntry(FMainFrameCommands::Get().VisitSupportWebSite);
+		HelpOnlineSection.AddMenuEntry(FMainFrameCommands::Get().VisitForums);
+		HelpOnlineSection.AddMenuEntry(FMainFrameCommands::Get().VisitSearchForAnswersPage);
+		HelpOnlineSection.AddMenuEntry(FMainFrameCommands::Get().VisitWiki);
 
 
 		const FText SupportWebSiteLabel = NSLOCTEXT("MainHelpMenu", "VisitUnrealEngineSupportWebSite", "Unreal Engine Support Web Site...");
 
-		MenuBuilder.AddMenuSeparator("EpicGamesHelp");
-		MenuBuilder.AddMenuEntry(FMainFrameCommands::Get().VisitEpicGamesDotCom, "VisitEpicGamesDotCom");
+		HelpOnlineSection.AddMenuSeparator("EpicGamesHelp");
+		HelpOnlineSection.AddMenuEntry(FMainFrameCommands::Get().VisitEpicGamesDotCom);
 
-		MenuBuilder.AddMenuSeparator("Credits");
-		MenuBuilder.AddMenuEntry(FMainFrameCommands::Get().CreditsUnrealEd);
+		HelpOnlineSection.AddMenuSeparator("Credits");
+		HelpOnlineSection.AddMenuEntry(FMainFrameCommands::Get().CreditsUnrealEd);
 	}
-	MenuBuilder.EndSection();
 
 #if !PLATFORM_MAC // Handled by app's menu in menu bar
-	MenuBuilder.BeginSection("HelpApplication", NSLOCTEXT("MainHelpMenu", "Application", "Application"));
+	FEditorMenuSection& HelpApplicationSection = Menu->AddSection("HelpApplication", NSLOCTEXT("MainHelpMenu", "Application", "Application"));
 	{
 		const FText AboutWindowTitle = NSLOCTEXT("MainHelpMenu", "AboutUnrealEditor", "About Unreal Editor...");
 
-		MenuBuilder.AddMenuEntry(FMainFrameCommands::Get().AboutUnrealEd, "AboutUnrealEd", AboutWindowTitle);
+		HelpApplicationSection.AddMenuEntry(FMainFrameCommands::Get().AboutUnrealEd, AboutWindowTitle);
 	}
-	MenuBuilder.EndSection();
 #endif
 }
 
+#undef LOCTEXT_NAMESPACE
 
-TSharedRef<SWidget> FMainMenu::MakeMainMenu(const TSharedPtr<FTabManager>& TabManager, const TSharedRef< FExtender > Extender)
+TSharedRef<SWidget> FMainMenu::MakeMainMenu(const TSharedPtr<FTabManager>& TabManager, const FName MenuName, FEditorMenuContext& EditorMenuContext)
 {
-#define LOCTEXT_NAMESPACE "MainMenu"
-
 	// Cache all project names once
 	FMainFrameActionCallbacks::CacheProjectNames();
 
-	FMenuBarBuilder MenuBuilder(FMainFrameCommands::ActionList, Extender);
+	FMainMenu::RegisterMainMenu();
+
+	EditorMenuContext.AppendCommandList(FMainFrameCommands::ActionList);
+
+	USlateTabManagerContext* ContextObject = NewObject<USlateTabManagerContext>();
+	ContextObject->TabManager = TabManager;
+	EditorMenuContext.AddObject(ContextObject);
+
+	// Create the menu bar!
+	TSharedRef<SWidget> MenuBarWidget = UEditorMenuSubsystem::Get()->GenerateWidget(MenuName, EditorMenuContext);
+	if (MenuBarWidget != SNullWidget::NullWidget)
 	{
-		// File
-		MenuBuilder.AddPullDownMenu( 
-			LOCTEXT("FileMenu", "File"),
-			LOCTEXT("FileMenu_ToolTip", "Open the file menu"),
-			FNewMenuDelegate::CreateStatic(&FMainMenu::FillFileMenu, Extender),
-			"File",
-			FName(TEXT("FileMenu"))
-		);
+		// Tell tab-manager about the multi-box for platforms with a global menu bar
+		TSharedRef<SMultiBoxWidget> MultiBoxWidget = StaticCastSharedRef<SMultiBoxWidget>(MenuBarWidget);
+		TabManager->SetMenuMultiBox(ConstCastSharedRef<FMultiBox>(MultiBoxWidget->GetMultiBox()));
+	}
 
-		// Edit
-		MenuBuilder.AddPullDownMenu( 
-			LOCTEXT("EditMenu", "Edit"),
-			LOCTEXT("EditMenu_ToolTip", "Open the edit menu"),
-			FNewMenuDelegate::CreateStatic(&FMainMenu::FillEditMenu, Extender, TabManager),
-			"Edit"
-			,
-			FName(TEXT("EditMenu"))
-		);
+	return MenuBarWidget;
+}
 
-		// Window
-		MenuBuilder.AddPullDownMenu(
-			LOCTEXT("WindowMenu", "Window"),
-			LOCTEXT("WindowMenu_ToolTip", "Open new windows or tabs."),
-			FNewMenuDelegate::CreateStatic(&FMainMenu::FillWindowMenu, Extender, TabManager),
-			"Window"
-		);
+void FMainMenu::RegisterMainMenu()
+{
+#define LOCTEXT_NAMESPACE "MainMenu"
 
-		// Help
-		MenuBuilder.AddPullDownMenu( 
-			LOCTEXT("HelpMenu", "Help"),
-			LOCTEXT("HelpMenu_ToolTip", "Open the help menu"),
-			FNewMenuDelegate::CreateStatic(&FMainMenu::FillHelpMenu, Extender),
-			"Help"
-		);
+	static const FName MainMenuName("MainFrame.MainMenu");
+	UEditorMenuSubsystem* EditorMenus = UEditorMenuSubsystem::Get();
+	if (EditorMenus->IsMenuRegistered(MainMenuName))
+	{
+		return;
 	}
 
-	// Create the menu bar!
-	TSharedRef<SWidget> MenuBarWidget = MenuBuilder.MakeWidget();
+	RegisterFileMenu();
+	RegisterEditMenu();
+	RegisterWindowMenu();
+	RegisterHelpMenu();
+
+	UEditorMenu* MenuBar = EditorMenus->RegisterMenu(MainMenuName, NAME_None, EMultiBoxType::MenuBar);
+
+	MenuBar->AddSubMenu(
+		"MainMenu",
+		NAME_None,
+		"File",
+		LOCTEXT("FileMenu", "File"),
+		LOCTEXT("FileMenu_ToolTip", "Open the file menu")
+	);
+
+	MenuBar->AddSubMenu(
+		"MainMenu",
+		NAME_None,
+		"Edit",
+		LOCTEXT("EditMenu", "Edit"),
+		LOCTEXT("EditMenu_ToolTip", "Open the edit menu")
+		);
 
-	// Tell tab-manager about the multi-box for platforms with a global menu bar
-	TabManager->SetMenuMultiBox(MenuBuilder.GetMultiBox());
-	
-	return MenuBarWidget;
+	MenuBar->AddSubMenu(
+		"MainMenu",
+		NAME_None,
+		"Window",
+		LOCTEXT("WindowMenu", "Window"),
+		LOCTEXT("WindowMenu_ToolTip", "Open new windows or tabs.")
+		);
+
+	MenuBar->AddSubMenu(
+		"MainMenu",
+		NAME_None,
+		"Help",
+		LOCTEXT("HelpMenu", "Help"),
+		LOCTEXT("HelpMenu_ToolTip", "Open the help menu")
+	);
 }
 
 #undef LOCTEXT_NAMESPACE
@@ -355,142 +412,136 @@ TSharedRef<SWidget> FMainMenu::MakeMainMenu(const TSharedPtr<FTabManager>& TabMa
 
 #define LOCTEXT_NAMESPACE "MainTabMenu"
 
-TSharedRef< SWidget > FMainMenu::MakeMainTabMenu( const TSharedPtr<FTabManager>& TabManager, const TSharedRef< FExtender > UserExtender )
-{	struct Local
+void FMainMenu::RegisterFileProjectMenu()
+{
+	if (!GetDefault<UEditorStyleSettings>()->bShowProjectMenus)
 	{
-		static void FillProjectMenuItems( FMenuBuilder& MenuBuilder )
-		{
-			MenuBuilder.BeginSection( "FileProject", LOCTEXT("ProjectHeading", "Project") );
-			{
-				MenuBuilder.AddMenuEntry( FMainFrameCommands::Get().NewProject );
-				MenuBuilder.AddMenuEntry( FMainFrameCommands::Get().OpenProject );
-
-				FText ShortIDEName = FSourceCodeNavigation::GetSelectedSourceCodeIDE();
-
-				MenuBuilder.AddMenuEntry( FMainFrameCommands::Get().AddCodeToProject,
-					NAME_None,
-					TAttribute<FText>(),
-					FText::Format(LOCTEXT("AddCodeToProjectTooltip", "Adds C++ code to the project. The code can only be compiled if you have {0} installed."), ShortIDEName)
-				);
-
-				MenuBuilder.AddSubMenu(
-					LOCTEXT("PackageProjectSubMenuLabel", "Package Project"),
-					LOCTEXT("PackageProjectSubMenuToolTip", "Compile, cook and package your project and its content for distribution."),
-					FNewMenuDelegate::CreateStatic( &FPackageProjectMenu::MakeMenu ), false, FSlateIcon(FEditorStyle::GetStyleSetName(), "MainFrame.PackageProject")
-				);
-				/*
-				MenuBuilder.AddMenuEntry( FMainFrameCommands::Get().LocalizeProject,
-					NAME_None,
-					TAttribute<FText>(),
-					LOCTEXT("LocalizeProjectToolTip", "Gather text from your project and import/export translations.")
-					);
-					*/
-				/*
-				MenuBuilder.AddSubMenu(
-					LOCTEXT("CookProjectSubMenuLabel", "Cook Project"),
-					LOCTEXT("CookProjectSubMenuToolTip", "Cook your project content for debugging"),
-					FNewMenuDelegate::CreateStatic( &FCookContentMenu::MakeMenu ), false, FSlateIcon()
-				);
-				*/
-
-				FString SolutionPath;
-				if (FSourceCodeNavigation::DoesModuleSolutionExist())
-				{
-					MenuBuilder.AddMenuEntry( FMainFrameCommands::Get().RefreshCodeProject,
-						NAME_None,
-						FText::Format(LOCTEXT("RefreshCodeProjectLabel", "Refresh {0} Project"), ShortIDEName),
-						FText::Format(LOCTEXT("RefreshCodeProjectTooltip", "Refreshes your C++ code project in {0}."), ShortIDEName)
-					);
-
-					MenuBuilder.AddMenuEntry( FMainFrameCommands::Get().OpenIDE,
-						NAME_None,
-						FText::Format(LOCTEXT("OpenIDELabel", "Open {0}"), ShortIDEName),
-						FText::Format(LOCTEXT("OpenIDETooltip", "Opens your C++ code in {0}."), ShortIDEName)
-					);
-				}
-				else
-				{
-					MenuBuilder.AddMenuEntry( FMainFrameCommands::Get().RefreshCodeProject,
-						NAME_None,
-						FText::Format(LOCTEXT("GenerateCodeProjectLabel", "Generate {0} Project"), ShortIDEName),
-						FText::Format(LOCTEXT("GenerateCodeProjectTooltip", "Generates your C++ code project in {0}."), ShortIDEName)
-					);
-				}
-
-				// @hack GDC: this should be moved somewhere else and be less hacky
-				ITargetPlatform* RunningTargetPlatform = GetTargetPlatformManager()->GetRunningTargetPlatform();
+		return;
+	}
 
-				if (RunningTargetPlatform != nullptr)
-				{
-					const FName CookedPlatformName = *(RunningTargetPlatform->PlatformName() + TEXT("NoEditor"));
-					const FText CookedPlatformText = FText::FromString(RunningTargetPlatform->PlatformName());
+	UEditorMenuSubsystem* EditorMenus = UEditorMenuSubsystem::Get();
+	UEditorMenu* MainTabFileMenu = EditorMenus->ExtendMenu("MainFrame.MainTabMenu.File");
+	FEditorMenuSection& Section = MainTabFileMenu->AddSection("FileProject", LOCTEXT("ProjectHeading", "Project"), FEditorMenuInsert("FileLoadAndSave", EEditorMenuInsertType::After));
+
+	Section.AddMenuEntry( FMainFrameCommands::Get().NewProject );
+	Section.AddMenuEntry( FMainFrameCommands::Get().OpenProject );
+
+	FText ShortIDEName = FSourceCodeNavigation::GetSelectedSourceCodeIDE();
+
+	Section.AddMenuEntry( FMainFrameCommands::Get().AddCodeToProject,
+		TAttribute<FText>(),
+		FText::Format(LOCTEXT("AddCodeToProjectTooltip", "Adds C++ code to the project. The code can only be compiled if you have {0} installed."), ShortIDEName)
+	);
+
+	Section.AddEntry(FEditorMenuEntry::InitSubMenu(
+		"MainFrame.MainTabMenu.File",
+		"PackageProject",
+		LOCTEXT("PackageProjectSubMenuLabel", "Package Project"),
+		LOCTEXT("PackageProjectSubMenuToolTip", "Compile, cook and package your project and its content for distribution."),
+		FNewMenuDelegate::CreateStatic( &FPackageProjectMenu::MakeMenu ), false, FSlateIcon(FEditorStyle::GetStyleSetName(), "MainFrame.PackageProject")
+	));
+
+	/*
+	MenuBuilder.AddMenuEntry( FMainFrameCommands::Get().LocalizeProject,
+		NAME_None,
+		TAttribute<FText>(),
+		LOCTEXT("LocalizeProjectToolTip", "Gather text from your project and import/export translations.")
+		);
+		*/
+	/*
+	MenuBuilder.AddSubMenu(
+		LOCTEXT("CookProjectSubMenuLabel", "Cook Project"),
+		LOCTEXT("CookProjectSubMenuToolTip", "Cook your project content for debugging"),
+		FNewMenuDelegate::CreateStatic( &FCookContentMenu::MakeMenu ), false, FSlateIcon()
+	);
+	*/
+
+	FString SolutionPath;
+	if (FSourceCodeNavigation::DoesModuleSolutionExist())
+	{
+		Section.AddMenuEntry( FMainFrameCommands::Get().RefreshCodeProject,
+			FText::Format(LOCTEXT("RefreshCodeProjectLabel", "Refresh {0} Project"), ShortIDEName),
+			FText::Format(LOCTEXT("RefreshCodeProjectTooltip", "Refreshes your C++ code project in {0}."), ShortIDEName)
+		);
 
-					FUIAction Action(
-						FExecuteAction::CreateStatic(&FMainFrameActionCallbacks::CookContent, CookedPlatformName),
-						FCanExecuteAction::CreateStatic(&FMainFrameActionCallbacks::CookContentCanExecute, CookedPlatformName)
-					);
+		Section.AddMenuEntry( FMainFrameCommands::Get().OpenIDE,
+			FText::Format(LOCTEXT("OpenIDELabel", "Open {0}"), ShortIDEName),
+			FText::Format(LOCTEXT("OpenIDETooltip", "Opens your C++ code in {0}."), ShortIDEName)
+		);
+	}
+	else
+	{
+		Section.AddMenuEntry( FMainFrameCommands::Get().RefreshCodeProject,
+			FText::Format(LOCTEXT("GenerateCodeProjectLabel", "Generate {0} Project"), ShortIDEName),
+			FText::Format(LOCTEXT("GenerateCodeProjectTooltip", "Generates your C++ code project in {0}."), ShortIDEName)
+		);
+	}
 
-					MenuBuilder.AddMenuEntry(
-						FText::Format(LOCTEXT("CookContentForPlatform", "Cook Content for {0}"), CookedPlatformText),
-						FText::Format(LOCTEXT("CookContentForPlatformTooltip", "Cook your game content for debugging on the {0} platform"), CookedPlatformText),
-						FSlateIcon(),
-						Action
-					);
-				}
-			}
-			MenuBuilder.EndSection();
-		}
+	Section.AddDynamicEntry("CookContentForPlatform", FNewEditorMenuSectionDelegate::CreateLambda([](FEditorMenuSection& InSection)
+	{
+		// @hack GDC: this should be moved somewhere else and be less hacky
+		ITargetPlatform* RunningTargetPlatform = GetTargetPlatformManager()->GetRunningTargetPlatform();
 
-		static void FillRecentFileAndExitMenuItems( FMenuBuilder& MenuBuilder )
+		if (RunningTargetPlatform != nullptr)
 		{
-			MenuBuilder.BeginSection("FileRecentFiles");
-			{
-				if (GetDefault<UEditorStyleSettings>()->bShowProjectMenus && FMainFrameActionCallbacks::ProjectNames.Num() > 0)
-				{
-					MenuBuilder.AddSubMenu(
-						LOCTEXT("SwitchProjectSubMenu", "Recent Projects"),
-						LOCTEXT("SwitchProjectSubMenu_ToolTip", "Select a project to switch to"),
-						FNewMenuDelegate::CreateStatic(&FRecentProjectsMenu::MakeMenu), false, FSlateIcon(FEditorStyle::GetStyleSetName(), "MainFrame.RecentProjects")
-						);
-				}
-			}
-			MenuBuilder.EndSection();
+			const FName CookedPlatformName = *(RunningTargetPlatform->PlatformName() + TEXT("NoEditor"));
+			const FText CookedPlatformText = FText::FromString(RunningTargetPlatform->PlatformName());
 
-#if !PLATFORM_MAC // Handled by app's menu in menu bar
-			MenuBuilder.AddMenuSeparator();
-			MenuBuilder.AddMenuEntry( FMainFrameCommands::Get().Exit, "Exit" );
-#endif
+			FUIAction Action(
+				FExecuteAction::CreateStatic(&FMainFrameActionCallbacks::CookContent, CookedPlatformName),
+				FCanExecuteAction::CreateStatic(&FMainFrameActionCallbacks::CookContentCanExecute, CookedPlatformName)
+			);
+
+			InSection.AddMenuEntry(
+				"CookContentForPlatform",
+				FText::Format(LOCTEXT("CookContentForPlatform", "Cook Content for {0}"), CookedPlatformText),
+				FText::Format(LOCTEXT("CookContentForPlatformTooltip", "Cook your game content for debugging on the {0} platform"), CookedPlatformText),
+				FSlateIcon(),
+				Action
+			);
 		}
-	};
+	}));
+}
 
-	FExtensibilityManager ExtensibilityManager;
+void FMainMenu::RegisterRecentFileAndExitMenuItems()
+{
+	UEditorMenuSubsystem* EditorMenus = UEditorMenuSubsystem::Get();
+	UEditorMenu* MainTabFileMenu = EditorMenus->RegisterMenu("MainFrame.MainTabMenu.File", "MainFrame.MainMenu.File");
 
-	ExtensibilityManager.AddExtender( UserExtender );
 	{
-		TSharedRef< FExtender > Extender( new FExtender() );
-	
-		IMainFrameModule& MainFrameModule = FModuleManager::GetModuleChecked<IMainFrameModule>("MainFrame");
-
-		if (GetDefault<UEditorStyleSettings>()->bShowProjectMenus)
+		FEditorMenuSection& Section = MainTabFileMenu->AddSection("FileRecentFiles");
+		if (GetDefault<UEditorStyleSettings>()->bShowProjectMenus && FMainFrameActionCallbacks::ProjectNames.Num() > 0)
 		{
-			Extender->AddMenuExtension(
-				"FileLoadAndSave",
-				EExtensionHook::After,
-				MainFrameModule.GetMainFrameCommandBindings(),
-				FMenuExtensionDelegate::CreateStatic(&Local::FillProjectMenuItems));
+			Section.AddEntry(FEditorMenuEntry::InitSubMenu(
+				"MainFrame.MainTabMenu.File",
+				"RecentProjects",
+				LOCTEXT("SwitchProjectSubMenu", "Recent Projects"),
+				LOCTEXT("SwitchProjectSubMenu_ToolTip", "Select a project to switch to"),
+				FNewEditorMenuDelegate::CreateStatic(&FRecentProjectsMenu::MakeMenu),
+				false,
+				FSlateIcon(FEditorStyle::GetStyleSetName(), "MainFrame.RecentProjects")
+			));
 		}
+	}
 
-		Extender->AddMenuExtension(
-			"FileLoadAndSave",
-			EExtensionHook::After,
-			MainFrameModule.GetMainFrameCommandBindings(),
-			FMenuExtensionDelegate::CreateStatic( &Local::FillRecentFileAndExitMenuItems ) );
+#if !PLATFORM_MAC // Handled by app's menu in menu bar
+	{
+		FEditorMenuSection& Section = MainTabFileMenu->AddSection("Exit");
+		Section.AddMenuEntry( FMainFrameCommands::Get().Exit );
+	}
+#endif
+}
 
-		ExtensibilityManager.AddExtender( Extender );
+TSharedRef< SWidget > FMainMenu::MakeMainTabMenu( const TSharedPtr<FTabManager>& TabManager, const FName MenuName, FEditorMenuContext& EditorMenuContext )
+{
+	if (GetDefault<UEditorStyleSettings>()->bShowProjectMenus)
+	{
+		IMainFrameModule& MainFrameModule = FModuleManager::GetModuleChecked<IMainFrameModule>("MainFrame");
+		EditorMenuContext.AppendCommandList(MainFrameModule.GetMainFrameCommandBindings());
 	}
 
 
-	TSharedRef< SWidget > MenuBarWidget = FMainMenu::MakeMainMenu( TabManager, ExtensibilityManager.GetAllExtenders().ToSharedRef() );
+	TSharedRef< SWidget > MenuBarWidget = FMainMenu::MakeMainMenu( TabManager, MenuName, EditorMenuContext );
 
 	return MenuBarWidget;
 }
diff --git a/Engine/Source/Editor/MainFrame/Private/Menus/MainMenu.h b/Engine/Source/Editor/MainFrame/Private/Menus/MainMenu.h
--- a/Engine/Source/Editor/MainFrame/Private/Menus/MainMenu.h
+++ b/Engine/Source/Editor/MainFrame/Private/Menus/MainMenu.h
@@ -11,6 +11,7 @@
 #include "ILocalizationDashboardModule.h"
 
 class FMenuBuilder;
+struct FEditorMenuContext;
 
 /**
  * Unreal editor main frame Slate widget
@@ -23,56 +24,58 @@ public:
 	 * Static: Creates a widget for the main menu bar.
 	 *
 	 * @param TabManager Create the workspace menu based on this tab manager.
-	 * @param Extender Extensibility support for the menu.
-	 * @return New widget
+	 * @param MenuName Identifier associated with the menu.
+	 * @return EditorMenuContext Context containing state.
 	 */
-	static TSharedRef<SWidget> MakeMainMenu( const TSharedPtr<FTabManager>& TabManager, const TSharedRef<FExtender> Extender );
+	static TSharedRef<SWidget> MakeMainMenu( const TSharedPtr<FTabManager>& TabManager, const FName MenuName, FEditorMenuContext& EditorMenuContext );
 
 	/**
 	 * Static: Creates a widget for the main tab's menu bar.  This is just like the main menu bar, but also includes.
 	 * some "project level" menu items that we don't want propagated to most normal menus.
 	 *
 	 * @param TabManager Create the workspace menu based on this tab manager.
-	 * @param UserExtender Extensibility support for the menu.
-	 * @return New widget
+	 * @param MenuName Identifier associated with the menu.
+	 * @return EditorMenuContext Context containing state.
 	 */
-	static TSharedRef<SWidget> MakeMainTabMenu( const TSharedPtr<FTabManager>& TabManager, const TSharedRef<FExtender> UserExtender );
+	static TSharedRef<SWidget> MakeMainTabMenu( const TSharedPtr<FTabManager>& TabManager, const FName MenuName, FEditorMenuContext& EditorMenuContext );
+
+	/**
+	 * Static: Registers main menu with menu system.
+	 */
+	static void RegisterMainMenu();
 
 protected:
 
 	/**
 	 * Called to fill the file menu's content.
 	 *
-	 * @param MenuBuilder The multi-box builder that should be filled with content for this pull-down menu.
-	 * @param Extender Extensibility support for this menu.
 	 */
-	static void FillFileMenu( FMenuBuilder& MenuBuilder, const TSharedRef<FExtender> Extender );
+	static void RegisterFileMenu();
 
 	/**
 	 * Called to fill the edit menu's content.
-	 *
-	 * @param MenuBuilder The multi-box builder that should be filled with content for this pull-down menu.
-	 * @param Extender Extensibility support for this menu.
-	 * @param TabManager A Tab Manager from which to populate tab spawner menu items.
 	 */
-	static void FillEditMenu( FMenuBuilder& MenuBuilder, const TSharedRef<FExtender> Extender, const TSharedPtr<FTabManager> TabManager );
+	static void RegisterEditMenu();
 
 	/**
 	 * Called to fill the app menu's content.
-	 *
-	 * @param MenuBuilder The multi-box builder that should be filled with content for this pull-down menu.
-	 * @param Extender Extensibility support for this menu.
-	 * @param TabManager A Tab Manager from which to populate tab spawner menu items.
 	 */
-	static void FillWindowMenu( FMenuBuilder& MenuBuilder, const TSharedRef<FExtender> Extender, const TSharedPtr<FTabManager> TabManager );
+	static void RegisterWindowMenu();
 
 	/**
 	 * Called to fill the help menu's content.
-	 *
-	 * @param MenuBuilder The multi-box builder that should be filled with content for this pull-down menu.
-	 * @param Extender Extensibility support for this menu.
 	 */
-	static void FillHelpMenu( FMenuBuilder& MenuBuilder, const TSharedRef<FExtender> Extender );
+	static void RegisterHelpMenu();
+
+	/**
+	 * Called to fill the file menu's content.
+	 */
+	static void RegisterFileProjectMenu();
+
+	/**
+	 * Called to fill the file menu's recent and exit content.
+	 */
+	static void RegisterRecentFileAndExitMenuItems();
 
 private:
 	/** 
diff --git a/Engine/Source/Editor/MainFrame/Private/Menus/RecentProjectsMenu.h b/Engine/Source/Editor/MainFrame/Private/Menus/RecentProjectsMenu.h
--- a/Engine/Source/Editor/MainFrame/Private/Menus/RecentProjectsMenu.h
+++ b/Engine/Source/Editor/MainFrame/Private/Menus/RecentProjectsMenu.h
@@ -4,7 +4,7 @@
 
 #include "CoreMinimal.h"
 #include "Misc/Paths.h"
-#include "Framework/MultiBox/MultiBoxBuilder.h"
+#include "EditorMenuSubsystem.h"
 #include "Frame/MainFrameActions.h"
 #include "HAL/FileManager.h"
 
@@ -20,9 +20,9 @@ public:
 	/**
 	 * Creates the menu.
 	 *
-	 * @param MenuBuilder The builder for the menu that owns this menu.
+	 * @param Menu	The menu being populated.
 	 */
-	static void MakeMenu( FMenuBuilder& MenuBuilder )
+	static void MakeMenu( UEditorMenu* Menu )
 	{
 		for ( int32 ProjectIndex = 0; ProjectIndex < FMainFrameActionCallbacks::ProjectNames.Num() && ProjectIndex < FMainFrameCommands::Get().SwitchProjectCommands.Num(); ++ProjectIndex )
 		{
@@ -36,9 +36,10 @@ public:
 				continue;
 			}
 
+			FEditorMenuSection& Section = Menu->FindOrAddSection("Recent");
 			const FText DisplayName = FText::FromString( FPaths::GetBaseFilename(*ProjectName) );
 			const FText Tooltip = FText::FromString( IFileManager::Get().ConvertToAbsolutePathForExternalAppForRead(*ProjectName) );
-			MenuBuilder.AddMenuEntry( FMainFrameCommands::Get().SwitchProjectCommands[ ProjectIndex ], NAME_None, DisplayName, Tooltip );
+			Section.AddMenuEntry( FMainFrameCommands::Get().SwitchProjectCommands[ ProjectIndex ], DisplayName, Tooltip );
 		}
 	}
 };
diff --git a/Engine/Source/Editor/MainFrame/Private/Menus/SettingsMenu.h b/Engine/Source/Editor/MainFrame/Private/Menus/SettingsMenu.h
--- a/Engine/Source/Editor/MainFrame/Private/Menus/SettingsMenu.h
+++ b/Engine/Source/Editor/MainFrame/Private/Menus/SettingsMenu.h
@@ -12,6 +12,7 @@
 #include "ISettingsContainer.h"
 #include "ISettingsSection.h"
 
+#include "EditorMenuSubsystem.h"
 
 #define LOCTEXT_NAMESPACE "FSettingsMenu"
 
@@ -28,7 +29,7 @@ public:
 	 * @param MenuBuilder The builder for the menu that owns this menu.
 	 * @param SettingsContainerName The name of the settings container to create the menu for.
 	 */
-	static void MakeMenu( FMenuBuilder& MenuBuilder, FName SettingsContainerName )
+	static void MakeMenu( UEditorMenu* MenuBuilder, FName SettingsContainerName )
 	{
 		ISettingsModule* SettingsModule = FModuleManager::GetModulePtr<ISettingsModule>("Settings");
 
@@ -58,7 +59,7 @@ public:
 
 			if (SettingsCategory->GetSections(SettingsSections) > 0)
 			{
-				MenuBuilder.BeginSection(SettingsCategory->GetName(), SettingsCategory->GetDisplayName());
+				FEditorMenuSection& Section = MenuBuilder->AddSection(SettingsCategory->GetName(), SettingsCategory->GetDisplayName());
 
 				SettingsSections.Sort(
 					[](const ISettingsSectionPtr& First, const ISettingsSectionPtr& Second)
@@ -70,15 +71,14 @@ public:
 				{
 					ISettingsSectionPtr SettingsSection = SettingsSections[SectionIndex];
 
-					MenuBuilder.AddMenuEntry(
+					Section.AddMenuEntry(
+						SettingsSection->GetName(),
 						SettingsSection->GetDisplayName(),
 						SettingsSection->GetDescription(),
 						FSlateIcon(),
 						FUIAction(FExecuteAction::CreateStatic(&FSettingsMenu::OpenSettings, SettingsContainerName, SettingsCategory->GetName(), SettingsSection->GetName()))
 					);
 				}
-
-				MenuBuilder.EndSection();
 			}
 		}
 	}
diff --git a/Engine/Source/Editor/MainFrame/Public/Interfaces/IMainFrameModule.h b/Engine/Source/Editor/MainFrame/Public/Interfaces/IMainFrameModule.h
--- a/Engine/Source/Editor/MainFrame/Public/Interfaces/IMainFrameModule.h
+++ b/Engine/Source/Editor/MainFrame/Public/Interfaces/IMainFrameModule.h
@@ -12,6 +12,7 @@
 
 class FTabManager;
 class SWindow;
+struct FEditorMenuContext;
 
 /**
  * Data needed to display a Developer tools in the status bar.
@@ -52,7 +53,7 @@ public:
 	 *
 	 * @return	The newly-created menu widget
 	 */
-	virtual TSharedRef<SWidget> MakeMainMenu( const TSharedPtr<FTabManager>& TabManager, const TSharedRef< FExtender > Extender ) const = 0;
+	virtual TSharedRef<SWidget> MakeMainMenu( const TSharedPtr<FTabManager>& TabManager, const FName MenuName, FEditorMenuContext& EditorMenuContext ) const = 0;
 
 	/**
 	 * Generates a menu that's just like the "main menu" widget above, except it also includes some infrequently used commands
@@ -63,7 +64,7 @@ public:
 	 *
 	 * @return	The newly-created menu widget
 	 */
-	virtual TSharedRef<SWidget> MakeMainTabMenu( const TSharedPtr<FTabManager>& TabManager, const TSharedRef< FExtender > Extender ) const = 0;
+	virtual TSharedRef<SWidget> MakeMainTabMenu( const TSharedPtr<FTabManager>& TabManager, const FName MenuName, FEditorMenuContext& EditorMenuContext ) const = 0;
 
 	/**
 	 * Generates a menu for status and developer
diff --git a/Engine/Source/Editor/UnrealEd/Private/Kismet2/DebuggerCommands.cpp b/Engine/Source/Editor/UnrealEd/Private/Kismet2/DebuggerCommands.cpp
--- a/Engine/Source/Editor/UnrealEd/Private/Kismet2/DebuggerCommands.cpp
+++ b/Engine/Source/Editor/UnrealEd/Private/Kismet2/DebuggerCommands.cpp
@@ -819,17 +819,7 @@ TSharedRef< SWidget > FPlayWorldCommands::GeneratePlayMenuContent( TSharedRef<FU
 {
 	// Get all menu extenders for this context menu from the level editor module
 	FLevelEditorModule& LevelEditorModule = FModuleManager::GetModuleChecked<FLevelEditorModule>(TEXT("LevelEditor"));
-	TArray<FLevelEditorModule::FLevelEditorMenuExtender> MenuExtenderDelegates = LevelEditorModule.GetAllLevelEditorToolbarPlayMenuExtenders();
-
-	TArray<TSharedPtr<FExtender>> Extenders;
-	for ( int32 i = 0; i < MenuExtenderDelegates.Num(); ++i )
-	{
-		if ( MenuExtenderDelegates[i].IsBound() )
-		{
-			Extenders.Add(MenuExtenderDelegates[i].Execute(InCommandList));
-		}
-	}
-	TSharedPtr<FExtender> MenuExtender = FExtender::Combine(Extenders);
+	TSharedPtr<FExtender> MenuExtender = LevelEditorModule.AssembleExtenders(InCommandList, LevelEditorModule.GetAllLevelEditorToolbarPlayMenuExtenders());
 
 	struct FLocal
 	{
diff --git a/Engine/Source/Editor/UnrealEd/Private/Toolkits/AssetEditorToolkit.cpp b/Engine/Source/Editor/UnrealEd/Private/Toolkits/AssetEditorToolkit.cpp
--- a/Engine/Source/Editor/UnrealEd/Private/Toolkits/AssetEditorToolkit.cpp
+++ b/Engine/Source/Editor/UnrealEd/Private/Toolkits/AssetEditorToolkit.cpp
@@ -27,6 +27,8 @@
 #include "IAssetTools.h"
 #include "IAssetTypeActions.h"
 #include "AssetToolsModule.h"
+#include "Toolkits/AssetEditorToolkitMenuContext.h"
+#include "EditorMenuSubsystem.h"
 
 #define LOCTEXT_NAMESPACE "AssetEditorToolkit"
 
@@ -957,18 +959,60 @@ void FAssetEditorToolkit::FillDefaultHelpMenuCommands( FMenuBuilder& MenuBuilder
 	MenuBuilder.AddMenuEntry(FGlobalEditorCommonCommands::Get().OpenDocumentation, NAME_None, ToolTip);
 }
 
+FName FAssetEditorToolkit::GetEditorMenuAppName() const
+{
+	if (IsSimpleAssetEditor() && EditingObjects.Num() == 1 && EditingObjects[0])
+	{
+		return *(EditingObjects[0]->GetClass()->GetFName().ToString() + TEXT("Editor"));
+	}
+
+	return GetToolkitFName();
+}
+
+FName FAssetEditorToolkit::GetEditorMenuToolbarName() const
+{
+	return *(TEXT("AssetEditor.") + GetEditorMenuAppName().ToString() + TEXT(".ToolBar"));
+}
+
+void FAssetEditorToolkit::RegisterMenus()
+{
+	static const FName DefaultToolBarName("AssetEditor.DefaultToolBar");
+	UEditorMenuSubsystem* EditorMenus = UEditorMenuSubsystem::Get();
+	if (!EditorMenus->IsMenuRegistered(DefaultToolBarName))
+	{
+		UEditorMenu* ToolbarBuilder = EditorMenus->RegisterMenu(DefaultToolBarName, NAME_None, EMultiBoxType::ToolBar);
+		{
+			FEditorMenuSection& Section = ToolbarBuilder->AddSection("Asset");
+			Section.AddEntry(FEditorMenuEntry::InitToolBarButton(FAssetEditorCommonCommands::Get().SaveAsset));
+			Section.AddEntry(FEditorMenuEntry::InitToolBarButton(FGlobalEditorCommonCommands::Get().FindInContentBrowser, LOCTEXT("FindInContentBrowserButton", "Browse")));
+		}
+	}
+}
+
 void FAssetEditorToolkit::GenerateToolbar()
 {
 	TSharedPtr<FExtender> Extender = FExtender::Combine(ToolbarExtenders);
 
-	FToolBarBuilder ToolbarBuilder( GetToolkitCommands(), FMultiBoxCustomization::AllowCustomization( GetToolkitFName() ), Extender, Orient_Horizontal, bIsToolbarUsingSmallIcons);
-	ToolbarBuilder.SetIsFocusable(bIsToolbarFocusable);
-	ToolbarBuilder.BeginSection("Asset");
+	RegisterMenus();
+
+	const FName ToolBarName = GetEditorMenuToolbarName();
+	UEditorMenuSubsystem* EditorMenus = UEditorMenuSubsystem::Get();
+	UEditorMenu* FoundMenu = EditorMenus->FindMenu(ToolBarName);
+	if (!FoundMenu || !FoundMenu->IsRegistered())
 	{
-		ToolbarBuilder.AddToolBarButton(FAssetEditorCommonCommands::Get().SaveAsset);
-		ToolbarBuilder.AddToolBarButton(FGlobalEditorCommonCommands::Get().FindInContentBrowser, NAME_None, LOCTEXT("FindInContentBrowserButton", "Browse"));
+		FoundMenu = EditorMenus->RegisterMenu(ToolBarName, "AssetEditor.DefaultToolBar", EMultiBoxType::ToolBar);
 	}
-	ToolbarBuilder.EndSection();
+
+	FEditorMenuContext MenuContext(GetToolkitCommands(), Extender);
+
+	UAssetEditorToolkitMenuContext* ToolkitMenuContext = NewObject<UAssetEditorToolkitMenuContext>(FoundMenu);
+	ToolkitMenuContext->Toolkit = AsShared();
+	MenuContext.AddObject(ToolkitMenuContext);
+
+	UEditorMenu* GeneratedToolbar = EditorMenus->GenerateMenu(ToolBarName, MenuContext);
+	GeneratedToolbar->bToolBarIsFocusable = bIsToolbarFocusable;
+	GeneratedToolbar->bToolBarForceSmallIcons = bIsToolbarUsingSmallIcons;
+	TSharedRef< class SWidget > ToolBarWidget = EditorMenus->GenerateWidget(GeneratedToolbar);
 
 	TSharedRef<SHorizontalBox> MiscWidgets = SNew(SHorizontalBox);
 
@@ -994,7 +1038,7 @@ void FAssetEditorToolkit::GenerateToolbar()
 			.AutoHeight()
 			.VAlign(VAlign_Bottom)
 			[
-				ToolbarBuilder.MakeWidget()
+				ToolBarWidget
 			]
 		]
 		+SHorizontalBox::Slot()
diff --git a/Engine/Source/Editor/UnrealEd/Private/Toolkits/SStandaloneAssetEditorToolkitHost.cpp b/Engine/Source/Editor/UnrealEd/Private/Toolkits/SStandaloneAssetEditorToolkitHost.cpp
--- a/Engine/Source/Editor/UnrealEd/Private/Toolkits/SStandaloneAssetEditorToolkitHost.cpp
+++ b/Engine/Source/Editor/UnrealEd/Private/Toolkits/SStandaloneAssetEditorToolkitHost.cpp
@@ -12,6 +12,8 @@
 #include "Widgets/Docking/SDockTab.h"
 #include "UObject/Package.h"
 
+#include "EditorMenuSubsystem.h"
+
 #define LOCTEXT_NAMESPACE "StandaloneAssetEditorToolkit"
 
 void SStandaloneAssetEditorToolkitHost::Construct( const SStandaloneAssetEditorToolkitHost::FArguments& InArgs, const TSharedPtr<FTabManager>& InTabManager, const FName InitAppName )
@@ -40,7 +42,12 @@ void SStandaloneAssetEditorToolkitHost::SetupInitialContent( const TSharedRef<FT
 	//						- May need to change arrangement to allow for wider tool bars (maybe displace grid settings too)
 	//				- In standalone, just draws under the toolkit's menu
 
-	
+	const FName AssetEditorMenuName = GetMenuName();
+	if (!UEditorMenuSubsystem::Get()->IsMenuRegistered(AssetEditorMenuName))
+	{
+		UEditorMenuSubsystem::Get()->RegisterMenu(AssetEditorMenuName, "MainFrame.MainMenu");
+	}
+
 	if (bCreateDefaultStandaloneMenu)
 	{
 		struct Local
@@ -172,13 +179,29 @@ void SStandaloneAssetEditorToolkitHost::RestoreFromLayout( const TSharedRef<FTab
 	];
 }
 
+FName SStandaloneAssetEditorToolkitHost::GetMenuName() const
+{
+	FName MenuAppName;
+	if (HostedAssetEditorToolkit.IsValid())
+	{
+		MenuAppName = HostedAssetEditorToolkit->GetEditorMenuAppName();
+	}
+	else
+	{
+		MenuAppName = AppName;
+	}
+
+	return *(FString(TEXT("AssetEditor.")) + MenuAppName.ToString() + TEXT(".MainMenu"));
+}
 
 void SStandaloneAssetEditorToolkitHost::GenerateMenus(bool bForceCreateMenu)
 {
 	if( bForceCreateMenu || DefaultMenuWidget != SNullWidget::NullWidget )
 	{
+		const FName AssetEditorMenuName = GetMenuName();
+		FEditorMenuContext EditorMenuContext(HostedAssetEditorToolkit->GetToolkitCommands(), FExtender::Combine(MenuExtenders).ToSharedRef());
 		IMainFrameModule& MainFrameModule = FModuleManager::LoadModuleChecked<IMainFrameModule>( "MainFrame" );
-		DefaultMenuWidget = MainFrameModule.MakeMainMenu( MyTabManager, FExtender::Combine(MenuExtenders).ToSharedRef() );
+		DefaultMenuWidget = MainFrameModule.MakeMainMenu( MyTabManager, AssetEditorMenuName, EditorMenuContext );
 
 		MenuWidgetContent->SetContent(DefaultMenuWidget.ToSharedRef());
 	}
diff --git a/Engine/Source/Editor/UnrealEd/Private/Toolkits/SStandaloneAssetEditorToolkitHost.h b/Engine/Source/Editor/UnrealEd/Private/Toolkits/SStandaloneAssetEditorToolkitHost.h
--- a/Engine/Source/Editor/UnrealEd/Private/Toolkits/SStandaloneAssetEditorToolkitHost.h
+++ b/Engine/Source/Editor/UnrealEd/Private/Toolkits/SStandaloneAssetEditorToolkitHost.h
@@ -85,6 +85,8 @@ public:
 private:
 	void OnTabClosed(TSharedRef<SDockTab> TabClosed) const;
 
+	FName GetMenuName() const;
+
 	/** Manages internal tab layout */
 	TSharedPtr<FTabManager> MyTabManager;
 
diff --git a/Engine/Source/Editor/UnrealEd/Public/Toolkits/AssetEditorToolkit.h b/Engine/Source/Editor/UnrealEd/Public/Toolkits/AssetEditorToolkit.h
--- a/Engine/Source/Editor/UnrealEd/Public/Toolkits/AssetEditorToolkit.h
+++ b/Engine/Source/Editor/UnrealEd/Public/Toolkits/AssetEditorToolkit.h
@@ -175,6 +175,9 @@ public:
 	{
 		return TabManager;
 	}
+
+	/** Registers default tool bar */
+	static void RegisterMenus();
 	
 	/** Makes a default asset editing toolbar */
 	void GenerateToolbar();
@@ -182,6 +185,10 @@ public:
 	/** Regenerates the menubar and toolbar widgets */
 	void RegenerateMenusAndToolbars();
 
+	/** Get name used by editor menu subsystem */
+	FName GetEditorMenuToolbarName() const;
+	FName GetEditorMenuAppName() const;
+
 	/** Called at the end of RegenerateMenusAndToolbars() */
 	virtual void PostRegenerateMenusAndToolbars() { }
 	
diff --git a/Engine/Source/Editor/UnrealEd/Public/Toolkits/IToolkit.h b/Engine/Source/Editor/UnrealEd/Public/Toolkits/IToolkit.h
--- a/Engine/Source/Editor/UnrealEd/Public/Toolkits/IToolkit.h
+++ b/Engine/Source/Editor/UnrealEd/Public/Toolkits/IToolkit.h
@@ -133,5 +133,8 @@ public:
 
 	/* Called when a toolkit has been brought to the 'front' */
 	virtual void ToolkitBroughtToFront() {};
+
+	/* Returns true if this is the default generic asset editor used by multiple classes */
+	virtual bool IsSimpleAssetEditor() const { return false; }
 };
 
diff --git a/Engine/Source/Editor/UnrealEd/Public/Toolkits/SimpleAssetEditor.h b/Engine/Source/Editor/UnrealEd/Public/Toolkits/SimpleAssetEditor.h
--- a/Engine/Source/Editor/UnrealEd/Public/Toolkits/SimpleAssetEditor.h
+++ b/Engine/Source/Editor/UnrealEd/Public/Toolkits/SimpleAssetEditor.h
@@ -41,7 +41,8 @@ public:
 	virtual FString GetWorldCentricTabPrefix() const override;
 	virtual FLinearColor GetWorldCentricTabColorScale() const override;
 	virtual bool IsPrimaryEditor() const override { return true; }
-	
+	virtual bool IsSimpleAssetEditor() const override { return true; }
+
 	/** Used to show or hide certain properties */
 	void SetPropertyVisibilityDelegate(FIsPropertyVisible InVisibilityDelegate);
 
