Rename Editor\EditorMenus\ module to Runtime\ToolMenus\ so that other modules can use the API without having to depend on editor modules.

diff --git a/Engine/Source/Editor/LevelEditor/Private/LevelEditorMenu.cpp b/Engine/Source/Editor/LevelEditor/Private/LevelEditorMenu.cpp
--- a/Engine/Source/Editor/LevelEditor/Private/LevelEditorMenu.cpp
+++ b/Engine/Source/Editor/LevelEditor/Private/LevelEditorMenu.cpp
@@ -15,7 +15,7 @@
 #include "MRUFavoritesList.h"
 #include "Framework/Commands/GenericCommands.h"
 #include "IDocumentation.h"
-#include "EditorMenuSubsystem.h"
+#include "ToolMenus.h"
 
 #define LOCTEXT_NAMESPACE "LevelEditorMenu"
 
@@ -25,9 +25,9 @@ void FLevelEditorMenu::RegisterLevelEditorMenus()
 	{
 		static void RegisterFileLoadAndSaveItems()
 		{
-			UEditorMenu* Menu = UEditorMenuSubsystem::Get()->ExtendMenu("LevelEditor.MainMenu.File");
-			FEditorMenuSection& Section = Menu->FindOrAddSection("FileLoadAndSave");
-			FEditorMenuInsert InsertPos(NAME_None, EEditorMenuInsertType::First);
+			UToolMenu* Menu = UToolMenus::Get()->ExtendMenu("LevelEditor.MainMenu.File");
+			FToolMenuSection& Section = Menu->FindOrAddSection("FileLoadAndSave");
+			FToolMenuInsert InsertPos(NAME_None, EToolMenuInsertType::First);
 
 			// New Level
 			Section.AddMenuEntry( FLevelEditorCommands::Get().NewLevel ).InsertPosition = InsertPos;
@@ -50,12 +50,12 @@ void FLevelEditorMenu::RegisterLevelEditorMenus()
 
 		static void FillFileRecentAndFavoriteFileItems()
 		{
-			UEditorMenu* Menu = UEditorMenuSubsystem::Get()->ExtendMenu("LevelEditor.MainMenu.File");
-			FEditorMenuInsert SectionInsertPos("FileRecentFiles", EEditorMenuInsertType::Before);
+			UToolMenu* Menu = UToolMenus::Get()->ExtendMenu("LevelEditor.MainMenu.File");
+			FToolMenuInsert SectionInsertPos("FileRecentFiles", EToolMenuInsertType::Before);
 
 			// Import/Export
 			{
-				FEditorMenuSection& Section = Menu->AddSection("FileActors", LOCTEXT("ImportExportHeading", "Actors"), SectionInsertPos);
+				FToolMenuSection& Section = Menu->AddSection("FileActors", LOCTEXT("ImportExportHeading", "Actors"), SectionInsertPos);
 				{
 					// Import Into Level
 					Section.AddMenuEntry(FLevelEditorCommands::Get().ImportScene);
@@ -87,12 +87,12 @@ void FLevelEditorMenu::RegisterLevelEditorMenus()
 						}
 					};
 
-					static void MakeFavoriteLevelMenu(UEditorMenu* InMenu)
+					static void MakeFavoriteLevelMenu(UToolMenu* InMenu)
 					{
 						// Add a button to add/remove the currently loaded map as a favorite
 						if (FLevelEditorActionCallbacks::ToggleFavorite_CanExecute())
 						{
-							FEditorMenuSection& Section = InMenu->AddSection("LevelEditorToggleFavorite");
+							FToolMenuSection& Section = InMenu->AddSection("LevelEditorToggleFavorite");
 							{
 								TAttribute<FText> ToggleFavoriteLabel;
 								ToggleFavoriteLabel.BindStatic(&Local::GetToggleFavoriteLabelText);
@@ -117,20 +117,20 @@ void FLevelEditorMenu::RegisterLevelEditorMenus()
 					}
 				};
 
-				FEditorMenuSection& Section = Menu->AddSection("FileFavoriteLevels", TAttribute<FText>(), SectionInsertPos);
+				FToolMenuSection& Section = Menu->AddSection("FileFavoriteLevels", TAttribute<FText>(), SectionInsertPos);
 
-				Section.AddDynamicEntry("FileFavoriteLevels", FNewEditorMenuSectionDelegate::CreateLambda([](FEditorMenuSection& InSection)
+				Section.AddDynamicEntry("FileFavoriteLevels", FNewToolMenuSectionDelegate::CreateLambda([](FToolMenuSection& InSection)
 				{
 					IMainFrameModule& MainFrameModule = FModuleManager::LoadModuleChecked<IMainFrameModule>("MainFrame");
 					const FMainMRUFavoritesList& RecentsAndFavorites = *MainFrameModule.GetMRUFavoritesList();
 					if (RecentsAndFavorites.GetNumItems() > 0)
 					{
-						InSection.AddEntry(FEditorMenuEntry::InitSubMenu(
+						InSection.AddEntry(FToolMenuEntry::InitSubMenu(
 							"LevelEditor.MainMenu.File",
 							"FavoriteLevelsSubMenu",
 							LOCTEXT("FavoriteLevelsSubMenu", "Favorite Levels"),
 							LOCTEXT("RecentLevelsSubMenu_ToolTip", "Select a level to load"),
-							FNewEditorMenuDelegate::CreateStatic(&FFavoriteLevelMenu::MakeFavoriteLevelMenu),
+							FNewToolMenuDelegate::CreateStatic(&FFavoriteLevelMenu::MakeFavoriteLevelMenu),
 							false,
 							FSlateIcon(FEditorStyle::GetStyleSetName(), "MainFrame.FavoriteLevels")
 						));
@@ -142,7 +142,7 @@ void FLevelEditorMenu::RegisterLevelEditorMenus()
 			{
 				struct FRecentLevelMenu
 				{
-					static void MakeRecentLevelMenu( UEditorMenu* InMenu )
+					static void MakeRecentLevelMenu( UToolMenu* InMenu )
 					{
 						const FMainMRUFavoritesList& MRUFavorites = *FModuleManager::LoadModuleChecked<IMainFrameModule>( "MainFrame" ).GetMRUFavoritesList();
 						const int32 NumRecents = MRUFavorites.GetNumItems();
@@ -162,19 +162,19 @@ void FLevelEditorMenu::RegisterLevelEditorMenus()
 					}
 				};
 
-				FEditorMenuSection& Section = Menu->AddSection("FileRecentLevels", TAttribute<FText>(), SectionInsertPos);
-				Section.AddDynamicEntry("FileRecentLevels", FNewEditorMenuSectionDelegate::CreateLambda([](FEditorMenuSection& InSection)
+				FToolMenuSection& Section = Menu->AddSection("FileRecentLevels", TAttribute<FText>(), SectionInsertPos);
+				Section.AddDynamicEntry("FileRecentLevels", FNewToolMenuSectionDelegate::CreateLambda([](FToolMenuSection& InSection)
 				{
 					IMainFrameModule& MainFrameModule = FModuleManager::LoadModuleChecked<IMainFrameModule>( "MainFrame" );
 					const FMainMRUFavoritesList& RecentsAndFavorites = *MainFrameModule.GetMRUFavoritesList();
 					if (RecentsAndFavorites.GetNumItems() > 0)
 					{
-						InSection.AddEntry(FEditorMenuEntry::InitSubMenu(
+						InSection.AddEntry(FToolMenuEntry::InitSubMenu(
 							"LevelEditor.MainMenu.File",
 							"RecentLevelsSubMenu",
 							LOCTEXT("RecentLevelsSubMenu", "Recent Levels"),
 							LOCTEXT("RecentLevelsSubMenu_ToolTip", "Select a level to load"),
-							FNewEditorMenuDelegate::CreateStatic(&FRecentLevelMenu::MakeRecentLevelMenu),
+							FNewToolMenuDelegate::CreateStatic(&FRecentLevelMenu::MakeRecentLevelMenu),
 							false,
 							FSlateIcon(FEditorStyle::GetStyleSetName(), "MainFrame.RecentLevels")
 						));
@@ -185,10 +185,10 @@ void FLevelEditorMenu::RegisterLevelEditorMenus()
 
 		static void ExtendEditMenu()
 		{
-			UEditorMenu* Menu = UEditorMenuSubsystem::Get()->RegisterMenu("LevelEditor.MainMenu.Edit", "MainFrame.MainMenu.Edit");
+			UToolMenu* Menu = UToolMenus::Get()->RegisterMenu("LevelEditor.MainMenu.Edit", "MainFrame.MainMenu.Edit");
 			// Edit Actor
 			{
-				FEditorMenuSection& Section = Menu->AddSection("EditMain", LOCTEXT("MainHeading", "Edit"), FEditorMenuInsert("EditHistory", EEditorMenuInsertType::After));
+				FToolMenuSection& Section = Menu->AddSection("EditMain", LOCTEXT("MainHeading", "Edit"), FToolMenuInsert("EditHistory", EToolMenuInsertType::After));
 				{		
 					Section.AddMenuEntry( FGenericCommands::Get().Cut );
 					Section.AddMenuEntry( FGenericCommands::Get().Copy );
@@ -202,8 +202,8 @@ void FLevelEditorMenu::RegisterLevelEditorMenus()
 
 		static void ExtendHelpMenu()
 		{
-			UEditorMenu* Menu = UEditorMenuSubsystem::Get()->RegisterMenu("LevelEditor.MainMenu.Help", "MainFrame.MainMenu.Help");
-			FEditorMenuSection& Section = Menu->AddSection("HelpBrowse", NSLOCTEXT("MainHelpMenu", "Browse", "Browse"), FEditorMenuInsert("BugReporting", EEditorMenuInsertType::Before));
+			UToolMenu* Menu = UToolMenus::Get()->RegisterMenu("LevelEditor.MainMenu.Help", "MainFrame.MainMenu.Help");
+			FToolMenuSection& Section = Menu->AddSection("HelpBrowse", NSLOCTEXT("MainHelpMenu", "Browse", "Browse"), FToolMenuInsert("BugReporting", EToolMenuInsertType::Before));
 			{
 				Section.AddMenuEntry( FLevelEditorCommands::Get().BrowseDocumentation );
 
@@ -218,10 +218,10 @@ void FLevelEditorMenu::RegisterLevelEditorMenus()
 		}
 	};
 
-	UEditorMenuSubsystem* EditorMenus = UEditorMenuSubsystem::Get();
-	EditorMenus->RegisterMenu("LevelEditor.MainMenu", "MainFrame.MainMenu", EMultiBoxType::MenuBar);
-	EditorMenus->RegisterMenu("LevelEditor.MainMenu.File", "MainFrame.MainTabMenu.File");
-	EditorMenus->RegisterMenu("LevelEditor.MainMenu.Window", "MainFrame.MainMenu.Window");
+	UToolMenus* ToolMenus = UToolMenus::Get();
+	ToolMenus->RegisterMenu("LevelEditor.MainMenu", "MainFrame.MainMenu", EMultiBoxType::MenuBar);
+	ToolMenus->RegisterMenu("LevelEditor.MainMenu.File", "MainFrame.MainTabMenu.File");
+	ToolMenus->RegisterMenu("LevelEditor.MainMenu.Window", "MainFrame.MainMenu.Window");
 
 	// Add level loading and saving menu items
 	Local::RegisterFileLoadAndSaveItems();
@@ -240,10 +240,10 @@ TSharedRef< SWidget > FLevelEditorMenu::MakeLevelEditorMenu( const TSharedPtr<FU
 {
 	FLevelEditorModule& LevelEditorModule = FModuleManager::LoadModuleChecked<FLevelEditorModule>("LevelEditor");
 	TSharedPtr<FExtender> Extenders = LevelEditorModule.GetMenuExtensibilityManager()->GetAllExtenders();
-	FEditorMenuContext EditorMenuContext(CommandList, Extenders.ToSharedRef());
+	FToolMenuContext ToolMenuContext(CommandList, Extenders.ToSharedRef());
 
 	IMainFrameModule& MainFrameModule = FModuleManager::LoadModuleChecked<IMainFrameModule>( "MainFrame" );
-	TSharedRef< SWidget > MenuBarWidget = MainFrameModule.MakeMainTabMenu( LevelEditor->GetTabManager(), "LevelEditor.MainMenu", EditorMenuContext );
+	TSharedRef< SWidget > MenuBarWidget = MainFrameModule.MakeMainTabMenu( LevelEditor->GetTabManager(), "LevelEditor.MainMenu", ToolMenuContext );
 
 	return MenuBarWidget;
 }
diff --git a/Engine/Source/Editor/LevelEditor/Private/LevelEditorToolBar.cpp b/Engine/Source/Editor/LevelEditor/Private/LevelEditorToolBar.cpp
--- a/Engine/Source/Editor/LevelEditor/Private/LevelEditorToolBar.cpp
+++ b/Engine/Source/Editor/LevelEditor/Private/LevelEditorToolBar.cpp
@@ -59,7 +59,7 @@
 #include "MaterialShaderQualitySettings.h"
 #include "RHIShaderPlatformDefinitions.inl"
 #include "LevelEditorMenuContext.h"
-#include "EditorMenuSubsystem.h"
+#include "ToolMenus.h"
 
 namespace LevelEditorActionHelpers
 {
@@ -1155,13 +1155,13 @@ void FLevelEditorToolBar::RegisterLevelEditorToolBar( const TSharedRef<FUIComman
 
 #define LOCTEXT_NAMESPACE "LevelEditorToolBar"
 
-	UEditorMenu* ToolbarBuilder = UEditorMenuSubsystem::Get()->RegisterMenu("LevelEditor.LevelEditorToolBar", NAME_None, EMultiBoxType::ToolBar);
+	UToolMenu* ToolbarBuilder = UToolMenus::Get()->RegisterMenu("LevelEditor.LevelEditorToolBar", NAME_None, EMultiBoxType::ToolBar);
 
 	{
-		FEditorMenuSection& Section = ToolbarBuilder->AddSection("File");
+		FToolMenuSection& Section = ToolbarBuilder->AddSection("File");
 
 		// Save All Levels
-		Section.AddEntry(FEditorMenuEntry::InitToolBarButton(
+		Section.AddEntry(FToolMenuEntry::InitToolBarButton(
 			FLevelEditorCommands::Get().Save,
 			TAttribute<FText>(),
 			TAttribute<FText>(),
@@ -1242,7 +1242,7 @@ void FLevelEditorToolBar::RegisterLevelEditorToolBar( const TSharedRef<FUIComman
 
 			FSourceControlStatus::CheckSourceControlStatus();
 
-			Section.AddEntry(FEditorMenuEntry::InitComboButton(
+			Section.AddEntry(FToolMenuEntry::InitComboButton(
 				"SourceControl",
 				FUIAction(),
 				FOnGetContent::CreateStatic(&FLevelEditorToolBar::GenerateSourceControlMenu, InCommandList),
@@ -1255,17 +1255,17 @@ void FLevelEditorToolBar::RegisterLevelEditorToolBar( const TSharedRef<FUIComman
 	}
 
 	{
-		FEditorMenuSection& Section = ToolbarBuilder->AddSection("Content");
-		Section.AddEntry(FEditorMenuEntry::InitToolBarButton(FLevelEditorCommands::Get().OpenContentBrowser, LOCTEXT( "ContentBrowser_Override", "Content" ), TAttribute<FText>(), TAttribute<FSlateIcon>(), "LevelToolbarContent"));
+		FToolMenuSection& Section = ToolbarBuilder->AddSection("Content");
+		Section.AddEntry(FToolMenuEntry::InitToolBarButton(FLevelEditorCommands::Get().OpenContentBrowser, LOCTEXT( "ContentBrowser_Override", "Content" ), TAttribute<FText>(), TAttribute<FSlateIcon>(), "LevelToolbarContent"));
 		if (FLauncherPlatformModule::Get()->CanOpenLauncher(true)) 
 		{
-			Section.AddEntry(FEditorMenuEntry::InitToolBarButton(FLevelEditorCommands::Get().OpenMarketplace, LOCTEXT("Marketplace_Override", "Marketplace"), TAttribute<FText>(), TAttribute<FSlateIcon>(), "LevelToolbarMarketplace"));
+			Section.AddEntry(FToolMenuEntry::InitToolBarButton(FLevelEditorCommands::Get().OpenMarketplace, LOCTEXT("Marketplace_Override", "Marketplace"), TAttribute<FText>(), TAttribute<FSlateIcon>(), "LevelToolbarMarketplace"));
 		}
 	}
 
-	FEditorMenuSection& SettingsSection = ToolbarBuilder->AddSection("Settings");
+	FToolMenuSection& SettingsSection = ToolbarBuilder->AddSection("Settings");
 	{
-		SettingsSection.AddEntry(FEditorMenuEntry::InitComboButton(
+		SettingsSection.AddEntry(FToolMenuEntry::InitComboButton(
 			"LevelToolbarQuickSettings",
 			FUIAction(),
 			FOnGetContent::CreateStatic(&FLevelEditorToolBar::GenerateQuickSettingsMenu, InCommandList),
@@ -1390,7 +1390,7 @@ void FLevelEditorToolBar::RegisterLevelEditorToolBar( const TSharedRef<FUIComman
 			}
 		};
 
-		SettingsSection.AddEntry(FEditorMenuEntry::InitToolBarButton( 
+		SettingsSection.AddEntry(FToolMenuEntry::InitToolBarButton( 
 			FLevelEditorCommands::Get().ToggleFeatureLevelPreview,
 			TAttribute<FText>::Create(&FPreviewModeFunctionality::GetPreviewModeText),
         	TAttribute<FText>::Create(&FPreviewModeFunctionality::GetPreviewModeTooltip),
@@ -1399,8 +1399,8 @@ void FLevelEditorToolBar::RegisterLevelEditorToolBar( const TSharedRef<FUIComman
 	}
 
 	{
-		FEditorMenuSection& Section = ToolbarBuilder->AddSection("Misc");
-		Section.AddEntry(FEditorMenuEntry::InitComboButton(
+		FToolMenuSection& Section = ToolbarBuilder->AddSection("Misc");
+		Section.AddEntry(FToolMenuEntry::InitComboButton(
 			"OpenBlueprint",
 			FUIAction(),
 			FOnGetContent::CreateStatic( &FLevelEditorToolBar::GenerateOpenBlueprintMenuContent, InCommandList, TWeakPtr<SLevelEditor>( InLevelEditor ) ),
@@ -1409,7 +1409,7 @@ void FLevelEditorToolBar::RegisterLevelEditorToolBar( const TSharedRef<FUIComman
 			FSlateIcon(FEditorStyle::GetStyleSetName(), "LevelEditor.OpenLevelBlueprint")
 			));
 
-		Section.AddEntry(FEditorMenuEntry::InitComboButton(
+		Section.AddEntry(FToolMenuEntry::InitComboButton(
 			"EditCinematics",
 			FUIAction(),
 			FOnGetContent::CreateStatic( &FLevelEditorToolBar::GenerateCinematicsMenuContent, InCommandList, TWeakPtr<SLevelEditor>( InLevelEditor ) ),
@@ -1418,16 +1418,16 @@ void FLevelEditorToolBar::RegisterLevelEditorToolBar( const TSharedRef<FUIComman
 			FSlateIcon(FEditorStyle::GetStyleSetName(), "LevelEditor.EditMatinee") 
 			));
 
-		Section.AddEntry(FEditorMenuEntry::InitToolBarButton( FLevelEditorCommands::Get().ToggleVR, LOCTEXT("ToggleVR", "VR Mode")) );
+		Section.AddEntry(FToolMenuEntry::InitToolBarButton( FLevelEditorCommands::Get().ToggleVR, LOCTEXT("ToggleVR", "VR Mode")) );
 	}
 
 	{
-		FEditorMenuSection& Section = ToolbarBuilder->AddSection("Compile");
+		FToolMenuSection& Section = ToolbarBuilder->AddSection("Compile");
 		// Build			
-		Section.AddEntry(FEditorMenuEntry::InitToolBarButton( FLevelEditorCommands::Get().Build, LOCTEXT("BuildAll", "Build")) );
+		Section.AddEntry(FToolMenuEntry::InitToolBarButton( FLevelEditorCommands::Get().Build, LOCTEXT("BuildAll", "Build")) );
 
 		// Build menu drop down
-		Section.AddEntry(FEditorMenuEntry::InitComboButton(
+		Section.AddEntry(FToolMenuEntry::InitComboButton(
 			"BuildComboButton",
 			FUIAction(),
 			FOnGetContent::CreateStatic( &FLevelEditorToolBar::GenerateBuildMenuContent, InCommandList ),
@@ -1440,7 +1440,7 @@ void FLevelEditorToolBar::RegisterLevelEditorToolBar( const TSharedRef<FUIComman
 		if ( FSourceCodeNavigation::IsCompilerAvailable() )
 		{
 			// Since we can always add new code to the project, only hide these buttons if we haven't done so yet
-			Section.AddEntry(FEditorMenuEntry::InitToolBarButton(
+			Section.AddEntry(FToolMenuEntry::InitToolBarButton(
 				"CompileButton",
 				FUIAction(
 					FExecuteAction::CreateStatic(&FLevelEditorActionCallbacks::RecompileGameCode_Clicked),
@@ -1453,7 +1453,7 @@ void FLevelEditorToolBar::RegisterLevelEditorToolBar( const TSharedRef<FUIComman
 				));
 
 #if WITH_LIVE_CODING
-			Section.AddEntry(FEditorMenuEntry::InitComboButton(
+			Section.AddEntry(FToolMenuEntry::InitComboButton(
 				"CompileComboButton",
 				FUIAction(
 					FExecuteAction(),
@@ -1470,7 +1470,7 @@ void FLevelEditorToolBar::RegisterLevelEditorToolBar( const TSharedRef<FUIComman
 		}
 	}
 	
-	ToolbarBuilder->AddDynamicSection("Game", FNewToolBarDelegateLegacy::CreateLambda([](FToolBarBuilder& InBuilder, UEditorMenu* InData)
+	ToolbarBuilder->AddDynamicSection("Game", FNewToolBarDelegateLegacy::CreateLambda([](FToolBarBuilder& InBuilder, UToolMenu* InData)
 	{
 		// Add the shared play-world commands that will be shown on the Kismet toolbar as well
 		FPlayWorldCommands::BuildToolbar(InBuilder, true);
@@ -1489,7 +1489,7 @@ TSharedRef< SWidget > FLevelEditorToolBar::MakeLevelEditorToolBar( const TShared
 #define LOCTEXT_NAMESPACE "LevelEditorToolBar"
 
 	FLevelEditorModule& LevelEditorModule = FModuleManager::LoadModuleChecked<FLevelEditorModule>("LevelEditor");
-	FEditorMenuContext MenuContext(InCommandList, LevelEditorModule.GetToolBarExtensibilityManager()->GetAllExtenders());
+	FToolMenuContext MenuContext(InCommandList, LevelEditorModule.GetToolBarExtensibilityManager()->GetAllExtenders());
 	ULevelEditorMenuContext* LevelEditorMenuContext = NewObject<ULevelEditorMenuContext>();
 	LevelEditorMenuContext->LevelEditor = InLevelEditor;
 	MenuContext.AddObject(LevelEditorMenuContext);
@@ -1501,7 +1501,7 @@ TSharedRef< SWidget > FLevelEditorToolBar::MakeLevelEditorToolBar( const TShared
 		.BorderImage( FEditorStyle::GetBrush("NoBorder") )
 		.IsEnabled( FSlateApplication::Get().GetNormalExecutionAttribute() )
 		[
-			UEditorMenuSubsystem::Get()->GenerateWidget("LevelEditor.LevelEditorToolBar", MenuContext)
+			UToolMenus::Get()->GenerateWidget("LevelEditor.LevelEditorToolBar", MenuContext)
 		];
 #undef LOCTEXT_NAMESPACE
 }
@@ -2047,19 +2047,19 @@ TSharedRef< SWidget > FLevelEditorToolBar::GenerateSourceControlMenu(TSharedRef<
 	FLevelEditorModule& LevelEditorModule = FModuleManager::GetModuleChecked<FLevelEditorModule>(TEXT("LevelEditor"));
 	TSharedPtr<FExtender> MenuExtender = LevelEditorModule.AssembleExtenders(InCommandList, LevelEditorModule.GetAllLevelEditorToolbarSourceControlMenuExtenders());
 
-	FEditorMenuContext Context(InCommandList, MenuExtender);
-	return UEditorMenuSubsystem::Get()->GenerateWidget("LevelEditor.LevelEditorToolBar.SourceControl", Context);
+	FToolMenuContext Context(InCommandList, MenuExtender);
+	return UToolMenus::Get()->GenerateWidget("LevelEditor.LevelEditorToolBar.SourceControl", Context);
 }
 
 void FLevelEditorToolBar::RegisterSourceControlMenu()
 {
 #define LOCTEXT_NAMESPACE "LevelToolBarSourceControlMenu"
 
-	UEditorMenu* Menu = UEditorMenuSubsystem::Get()->RegisterMenu("LevelEditor.LevelEditorToolBar.SourceControl");
+	UToolMenu* Menu = UToolMenus::Get()->RegisterMenu("LevelEditor.LevelEditorToolBar.SourceControl");
 	Menu->bShouldCloseWindowAfterMenuSelection = true;
-	FEditorMenuSection& Section = Menu->AddSection("SourceControlActions", LOCTEXT("SourceControlMenuHeadingActions", "Actions"));
+	FToolMenuSection& Section = Menu->AddSection("SourceControlActions", LOCTEXT("SourceControlMenuHeadingActions", "Actions"));
 
-	Section.AddDynamicEntry("ConnectToSourceControl", FNewEditorMenuSectionDelegate::CreateLambda([](FEditorMenuSection& InSection)
+	Section.AddDynamicEntry("ConnectToSourceControl", FNewToolMenuSectionDelegate::CreateLambda([](FToolMenuSection& InSection)
 	{
 		ISourceControlModule& SourceControlModule = ISourceControlModule::Get();
 		if (ISourceControlModule::Get().IsEnabled() && ISourceControlModule::Get().GetProvider().IsAvailable())
@@ -2265,28 +2265,28 @@ void FLevelEditorToolBar::OnOpenSubLevelBlueprint( ULevel* InLevel )
 TSharedRef< SWidget > FLevelEditorToolBar::GenerateCinematicsMenuContent( TSharedRef<FUICommandList> InCommandList, TWeakPtr<SLevelEditor> LevelEditorWeakPtr )
 {
 	FLevelEditorModule& LevelEditorModule = FModuleManager::LoadModuleChecked<FLevelEditorModule>("LevelEditor");
-	FEditorMenuContext Context(InCommandList, FExtender::Combine(LevelEditorModule.GetAllLevelEditorToolbarCinematicsMenuExtenders()));
+	FToolMenuContext Context(InCommandList, FExtender::Combine(LevelEditorModule.GetAllLevelEditorToolbarCinematicsMenuExtenders()));
 	ULevelEditorMenuContext* LevelEditorMenuContext = NewObject<ULevelEditorMenuContext>();
 	LevelEditorMenuContext->LevelEditor = LevelEditorWeakPtr;
 	Context.AddObject(LevelEditorMenuContext);
-	return UEditorMenuSubsystem::Get()->GenerateWidget("LevelEditor.LevelEditorToolBar.Cinematics", Context);
+	return UToolMenus::Get()->GenerateWidget("LevelEditor.LevelEditorToolBar.Cinematics", Context);
 }
 
 void FLevelEditorToolBar::RegisterCinematicsMenu()
 {
 #define LOCTEXT_NAMESPACE "LevelToolBarCinematicsMenu"
 
-	UEditorMenu* Menu = UEditorMenuSubsystem::Get()->RegisterMenu("LevelEditor.LevelEditorToolBar.Cinematics");
+	UToolMenu* Menu = UToolMenus::Get()->RegisterMenu("LevelEditor.LevelEditorToolBar.Cinematics");
 	Menu->bShouldCloseWindowAfterMenuSelection = true;
 
 	{
-		FEditorMenuSection& Section = Menu->AddSection("LevelEditorNewMatinee", LOCTEXT("MatineeMenuCombo_NewHeading", "New"));
+		FToolMenuSection& Section = Menu->AddSection("LevelEditorNewMatinee", LOCTEXT("MatineeMenuCombo_NewHeading", "New"));
 		Section.AddMenuEntry( FLevelEditorCommands::Get().AddMatinee, TAttribute<FText>(), TAttribute<FText>(), FSlateIcon(FEditorStyle::GetStyleSetName(), "ClassIcon.MatineeActor") );
 	}
 
 	//Add a heading to separate the existing cinematics from the 'Add New Cinematic Actor' button
-	FEditorMenuSection& ExistingCinematicSection = Menu->AddSection("LevelEditorExistingCinematic", LOCTEXT("CinematicMenuCombo_ExistingHeading", "Edit Existing Cinematic"));
-	ExistingCinematicSection.AddDynamicEntry("LevelEditorExistingCinematic", FNewEditorMenuSectionDelegate::CreateLambda([](FEditorMenuSection& InSection)
+	FToolMenuSection& ExistingCinematicSection = Menu->AddSection("LevelEditorExistingCinematic", LOCTEXT("CinematicMenuCombo_ExistingHeading", "Edit Existing Cinematic"));
+	ExistingCinematicSection.AddDynamicEntry("LevelEditorExistingCinematic", FNewToolMenuSectionDelegate::CreateLambda([](FToolMenuSection& InSection)
 	{
 		ULevelEditorMenuContext* FoundContext = InSection.Context.Find<ULevelEditorMenuContext>();
 		if (!FoundContext)
@@ -2337,7 +2337,7 @@ void FLevelEditorToolBar::RegisterCinematicsMenu()
 					FOnActorPicked::CreateStatic(&FLevelEditorToolBar::OnCinematicsActorPicked))
 			];
 
-		InSection.AddEntry(FEditorMenuEntry::InitWidget("LevelEditorExistingCinematic", MiniSceneOutliner, FText::GetEmpty(), true));
+		InSection.AddEntry(FToolMenuEntry::InitWidget("LevelEditorExistingCinematic", MiniSceneOutliner, FText::GetEmpty(), true));
 	}));
 
 #undef LOCTEXT_NAMESPACE
diff --git a/Engine/Source/Editor/LevelEditor/Private/LevelEditorToolBar.h b/Engine/Source/Editor/LevelEditor/Private/LevelEditorToolBar.h
--- a/Engine/Source/Editor/LevelEditor/Private/LevelEditorToolBar.h
+++ b/Engine/Source/Editor/LevelEditor/Private/LevelEditorToolBar.h
@@ -8,7 +8,7 @@
 #include "Framework/Commands/UICommandList.h"
 #include "Editor/LevelEditor/Private/SLevelEditor.h"
 
-class UEditorMenu;
+class UToolMenu;
 
 /**
  * Unreal level editor main toolbar
diff --git a/Engine/Source/Editor/LevelEditor/Private/SLevelEditor.cpp b/Engine/Source/Editor/LevelEditor/Private/SLevelEditor.cpp
--- a/Engine/Source/Editor/LevelEditor/Private/SLevelEditor.cpp
+++ b/Engine/Source/Editor/LevelEditor/Private/SLevelEditor.cpp
@@ -2,7 +2,7 @@
 
 #include "SLevelEditor.h"
 #include "Framework/MultiBox/MultiBoxExtender.h"
-#include "EditorMenuSubsystem.h"
+#include "ToolMenus.h"
 #include "Framework/Docking/LayoutService.h"
 #include "EditorModeRegistry.h"
 #include "EdMode.h"
@@ -699,26 +699,26 @@ TSharedRef<SDockTab> SLevelEditor::SpawnLevelEditorTab( const FSpawnTabArgs& Arg
 		InitOptions.bShowTransient = true;
 		InitOptions.Mode = ESceneOutlinerMode::ActorBrowsing;
 		{
-			UEditorMenuSubsystem* EditorMenus = UEditorMenuSubsystem::Get();
+			UToolMenus* ToolMenus = UToolMenus::Get();
 			static const FName MenuName = "LevelEditor.LevelEditorSceneOutliner.ContextMenu";
-			if (!EditorMenus->IsMenuRegistered(MenuName))
+			if (!ToolMenus->IsMenuRegistered(MenuName))
 			{
-				UEditorMenu* Menu = EditorMenus->RegisterMenu(MenuName, "SceneOutliner.DefaultContextMenuBase");
-				FEditorMenuSection& Section = Menu->AddDynamicSection("LevelEditorContextMenu", FNewEditorMenuDelegate::CreateLambda([](UEditorMenu* InMenu)
+				UToolMenu* Menu = ToolMenus->RegisterMenu(MenuName, "SceneOutliner.DefaultContextMenuBase");
+				FToolMenuSection& Section = Menu->AddDynamicSection("LevelEditorContextMenu", FNewToolMenuDelegate::CreateLambda([](UToolMenu* InMenu)
 				{
 					FName LevelContextMenuName = FLevelEditorContextMenu::GetContextMenuName(ELevelEditorMenuContext::SceneOutliner);
 					if (LevelContextMenuName != NAME_None)
 					{
 						// Extend the menu even if no actors selected, as Edit menu should always exist for scene outliner
-						UEditorMenu* OtherMenu = UEditorMenuSubsystem::Get()->GenerateMenu(LevelContextMenuName, InMenu->Context);
+						UToolMenu* OtherMenu = UToolMenus::Get()->GenerateMenu(LevelContextMenuName, InMenu->Context);
 						InMenu->Sections.Append(OtherMenu->Sections);
 					}
 				}));
-				Section.InsertPosition = FEditorMenuInsert("MainSection", EEditorMenuInsertType::Before);
+				Section.InsertPosition = FToolMenuInsert("MainSection", EToolMenuInsertType::Before);
 			}
 
 			TWeakPtr<SLevelEditor> WeakLevelEditor = SharedThis(this);
-			InitOptions.ModifyContextMenu.BindLambda([=](FName& OutMenuName, FEditorMenuContext& MenuContext)
+			InitOptions.ModifyContextMenu.BindLambda([=](FName& OutMenuName, FToolMenuContext& MenuContext)
 			{
 				OutMenuName = MenuName;
 
diff --git a/Engine/Source/Editor/LocalizationDashboard/Private/LocalizationDashboard.cpp b/Engine/Source/Editor/LocalizationDashboard/Private/LocalizationDashboard.cpp
--- a/Engine/Source/Editor/LocalizationDashboard/Private/LocalizationDashboard.cpp
+++ b/Engine/Source/Editor/LocalizationDashboard/Private/LocalizationDashboard.cpp
@@ -20,7 +20,7 @@
 #include "LocalizationSettings.h"
 #include "LocalizationDashboardSettings.h"
 #include "SSettingsEditorCheckoutNotice.h"
-#include "EditorMenuSubsystem.h"
+#include "ToolMenus.h"
 
 #define LOCTEXT_NAMESPACE "LocalizationDashboard"
 
@@ -145,9 +145,9 @@ void SLocalizationDashboard::Construct(const FArguments& InArgs, const TSharedPt
 
 	Layout = FLayoutSaveRestore::LoadFromConfig(GEditorLayoutIni, Layout);
 
-	FEditorMenuContext EditorMenuContext;
+	FToolMenuContext ToolMenuContext;
 	IMainFrameModule& MainFrameModule = FModuleManager::LoadModuleChecked<IMainFrameModule>( "MainFrame" );
-	const TSharedRef<SWidget> MenuWidget = MainFrameModule.MakeMainMenu( TabManager, "LocalizationDashboard.MainMenu", EditorMenuContext );
+	const TSharedRef<SWidget> MenuWidget = MainFrameModule.MakeMainMenu( TabManager, "LocalizationDashboard.MainMenu", ToolMenuContext );
 
 	FString ConfigFilePath;
 	ConfigFilePath = GetDefault<ULocalizationSettings>()->GetDefaultConfigFilename();
diff --git a/Engine/Source/Editor/MainFrame/Private/MainFrameModule.cpp b/Engine/Source/Editor/MainFrame/Private/MainFrameModule.cpp
--- a/Engine/Source/Editor/MainFrame/Private/MainFrameModule.cpp
+++ b/Engine/Source/Editor/MainFrame/Private/MainFrameModule.cpp
@@ -234,15 +234,15 @@ void FMainFrameModule::CreateDefaultMainFrame( const bool bStartImmersive, const
 }
 
 
-TSharedRef<SWidget> FMainFrameModule::MakeMainMenu( const TSharedPtr<FTabManager>& TabManager, const FName MenuName, FEditorMenuContext& EditorMenuContext ) const
+TSharedRef<SWidget> FMainFrameModule::MakeMainMenu(const TSharedPtr<FTabManager>& TabManager, const FName MenuName, FToolMenuContext& ToolMenuContext) const
 {
-	return FMainMenu::MakeMainMenu( TabManager, MenuName, EditorMenuContext );
+	return FMainMenu::MakeMainMenu(TabManager, MenuName, ToolMenuContext);
 }
 
 
-TSharedRef<SWidget> FMainFrameModule::MakeMainTabMenu( const TSharedPtr<FTabManager>& TabManager, const FName MenuName, FEditorMenuContext& EditorMenuContext ) const
+TSharedRef<SWidget> FMainFrameModule::MakeMainTabMenu(const TSharedPtr<FTabManager>& TabManager, const FName MenuName, FToolMenuContext& ToolMenuContext) const
 {
-	return FMainMenu::MakeMainTabMenu( TabManager, MenuName, EditorMenuContext );
+	return FMainMenu::MakeMainTabMenu(TabManager, MenuName, ToolMenuContext);
 }
 
 
diff --git a/Engine/Source/Editor/MainFrame/Private/MainFrameModule.h b/Engine/Source/Editor/MainFrame/Private/MainFrameModule.h
--- a/Engine/Source/Editor/MainFrame/Private/MainFrameModule.h
+++ b/Engine/Source/Editor/MainFrame/Private/MainFrameModule.h
@@ -23,8 +23,8 @@ public:
 	// IMainFrameModule interface
 
 	virtual void CreateDefaultMainFrame( const bool bStartImmersive, const bool bStartPIE ) override;
-	virtual TSharedRef<SWidget> MakeMainMenu( const TSharedPtr<FTabManager>& TabManager, const FName MenuName, FEditorMenuContext& EditorMenuContext ) const override;
-	virtual TSharedRef<SWidget> MakeMainTabMenu( const TSharedPtr<FTabManager>& TabManager, const FName MenuName, FEditorMenuContext& EditorMenuContext ) const override;
+	virtual TSharedRef<SWidget> MakeMainMenu(const TSharedPtr<FTabManager>& TabManager, const FName MenuName, FToolMenuContext& ToolMenuContext) const override;
+	virtual TSharedRef<SWidget> MakeMainTabMenu(const TSharedPtr<FTabManager>& TabManager, const FName MenuName, FToolMenuContext& ToolMenuContext) const override;
 	virtual TSharedRef<SWidget> MakeDeveloperTools( const TArray<FMainFrameDeveloperTool>& AdditionalTools ) const override;
 
 	virtual bool IsWindowInitialized( ) const override
diff --git a/Engine/Source/Editor/MainFrame/Private/Menus/MainMenu.cpp b/Engine/Source/Editor/MainFrame/Private/Menus/MainMenu.cpp
--- a/Engine/Source/Editor/MainFrame/Private/Menus/MainMenu.cpp
+++ b/Engine/Source/Editor/MainFrame/Private/Menus/MainMenu.cpp
@@ -23,7 +23,7 @@
 #include "Menus/SettingsMenu.h"
 #include "Menus/MainFrameTranslationEditorMenu.h"
 
-#include "EditorMenuSubsystem.h"
+#include "ToolMenus.h"
 
 #include "WorkspaceMenuStructure.h"
 #include "WorkspaceMenuStructureModule.h"
@@ -39,10 +39,10 @@
 
 void FMainMenu::RegisterFileMenu()
 {
-	UEditorMenuSubsystem* EditorMenus = UEditorMenuSubsystem::Get();
-	UEditorMenu* FileMenu = EditorMenus->RegisterMenu("MainFrame.MainMenu.File");
+	UToolMenus* ToolMenus = UToolMenus::Get();
+	UToolMenu* FileMenu = ToolMenus->RegisterMenu("MainFrame.MainMenu.File");
 
-	FEditorMenuSection& FileLoadAndSaveSection = FileMenu->AddSection("FileLoadAndSave", LOCTEXT("LoadSandSaveHeading", "Load and Save"));
+	FToolMenuSection& FileLoadAndSaveSection = FileMenu->AddSection("FileLoadAndSave", LOCTEXT("LoadSandSaveHeading", "Load and Save"));
 	{
 		// Open Asset...
 		FileLoadAndSaveSection.AddMenuEntry(FGlobalEditorCommonCommands::Get().SummonOpenAssetDialog);
@@ -53,7 +53,7 @@ void FMainMenu::RegisterFileMenu()
 		// Choose specific files to save
 		FileLoadAndSaveSection.AddMenuEntry(FMainFrameCommands::Get().ChooseFilesToSave);
 
-		FileLoadAndSaveSection.AddDynamicEntry("SourceControl", FNewEditorMenuSectionDelegate::CreateLambda([](FEditorMenuSection& InSection)
+		FileLoadAndSaveSection.AddDynamicEntry("SourceControl", FNewToolMenuSectionDelegate::CreateLambda([](FToolMenuSection& InSection)
 		{
 			if (ISourceControlModule::Get().IsEnabled() && ISourceControlModule::Get().GetProvider().IsAvailable())
 			{
@@ -78,9 +78,9 @@ void FMainMenu::RegisterFileMenu()
 
 void FMainMenu::RegisterEditMenu()
 {
-	UEditorMenu* EditMenu = UEditorMenuSubsystem::Get()->RegisterMenu("MainFrame.MainMenu.Edit");
+	UToolMenu* EditMenu = UToolMenus::Get()->RegisterMenu("MainFrame.MainMenu.Edit");
 	{
-		FEditorMenuSection& Section = EditMenu->AddSection("EditHistory", LOCTEXT("HistoryHeading", "History"));
+		FToolMenuSection& Section = EditMenu->AddSection("EditHistory", LOCTEXT("HistoryHeading", "History"));
 		struct Local
 		{
 			/** @return Returns a dynamic text string for Undo that contains the name of the action */
@@ -117,7 +117,7 @@ void FMainMenu::RegisterEditMenu()
 	}
 
 	{
-		FEditorMenuSection& Section = EditMenu->AddSection("EditLocalTabSpawners", LOCTEXT("ConfigurationHeading", "Configuration"));
+		FToolMenuSection& Section = EditMenu->AddSection("EditLocalTabSpawners", LOCTEXT("ConfigurationHeading", "Configuration"));
 		if (GetDefault<UEditorExperimentalSettings>()->bToolbarCustomization)
 		{
 			FUIAction ToggleMultiBoxEditMode(
@@ -135,7 +135,7 @@ void FMainMenu::RegisterEditMenu()
 				EUserInterfaceActionType::ToggleButton
 			);
 
-			Section.AddDynamicEntry("TabManager", FNewEditorMenuDelegateLegacy::CreateLambda([](FMenuBuilder& MenuBuilder, UEditorMenu* MenuData)
+			Section.AddDynamicEntry("TabManager", FNewToolMenuDelegateLegacy::CreateLambda([](FMenuBuilder& MenuBuilder, UToolMenu* MenuData)
 			{
 				if (USlateTabManagerContext* TabManagerContext = MenuData->FindContext<USlateTabManagerContext>())
 				{
@@ -151,22 +151,22 @@ void FMainMenu::RegisterEditMenu()
 
 		if (GetDefault<UEditorStyleSettings>()->bExpandConfigurationMenus)
 		{
-			Section.AddEntry(FEditorMenuEntry::InitSubMenu(
+			Section.AddEntry(FToolMenuEntry::InitSubMenu(
 				EditMenu->GetMenuName(),
 				"EditorPreferencesSubMenu",
 				LOCTEXT("EditorPreferencesSubMenuLabel", "Editor Preferences"),
 				LOCTEXT("EditorPreferencesSubMenuToolTip", "Configure the behavior and features of this Editor"),
-				FNewEditorMenuDelegate::CreateStatic(&FSettingsMenu::MakeMenu, FName("Editor")),
+				FNewToolMenuDelegate::CreateStatic(&FSettingsMenu::MakeMenu, FName("Editor")),
 				false,
 				FSlateIcon(FEditorStyle::GetStyleSetName(), "EditorPreferences.TabIcon")
 			));
 
-			Section.AddEntry(FEditorMenuEntry::InitSubMenu(
+			Section.AddEntry(FToolMenuEntry::InitSubMenu(
 				EditMenu->GetMenuName(),
 				"ProjectSettingsSubMenu",
 				LOCTEXT("ProjectSettingsSubMenuLabel", "Project Settings"),
 				LOCTEXT("ProjectSettingsSubMenuToolTip", "Change the settings of the currently loaded project"),
-				FNewEditorMenuDelegate::CreateStatic(&FSettingsMenu::MakeMenu, FName("Project")),
+				FNewToolMenuDelegate::CreateStatic(&FSettingsMenu::MakeMenu, FName("Project")),
 				false,
 				FSlateIcon(FEditorStyle::GetStyleSetName(), "ProjectSettings.TabIcon")
 			));
@@ -192,7 +192,7 @@ void FMainMenu::RegisterEditMenu()
 			);
 		}
 
-		Section.AddDynamicEntry("PluginsEditor", FNewEditorMenuDelegateLegacy::CreateLambda([](FMenuBuilder& InBuilder, UEditorMenu* InData)
+		Section.AddDynamicEntry("PluginsEditor", FNewToolMenuDelegateLegacy::CreateLambda([](FMenuBuilder& InBuilder, UToolMenu* InData)
 		{
 			//@todo The tab system needs to be able to be extendable by plugins [9/3/2013 Justin.Sargent]
 			if (IModularFeatures::Get().IsModularFeatureAvailable(EditorFeatures::PluginsEditor))
@@ -209,10 +209,10 @@ void FMainMenu::RegisterEditMenu()
 
 void FMainMenu::RegisterWindowMenu()
 {
-	UEditorMenu* Menu = UEditorMenuSubsystem::Get()->RegisterMenu("MainFrame.MainMenu.Window");
+	UToolMenu* Menu = UToolMenus::Get()->RegisterMenu("MainFrame.MainMenu.Window");
 
 	// Automatically populate tab spawners from TabManager
-	Menu->AddDynamicSection("TabManagerSection", FNewEditorMenuDelegateLegacy::CreateLambda([](FMenuBuilder& InBuilder, UEditorMenu* InData)
+	Menu->AddDynamicSection("TabManagerSection", FNewToolMenuDelegateLegacy::CreateLambda([](FMenuBuilder& InBuilder, UToolMenu* InData)
 	{
 		if (USlateTabManagerContext* TabManagerContext = InData->FindContext<USlateTabManagerContext>())
 		{
@@ -230,7 +230,7 @@ void FMainMenu::RegisterWindowMenu()
 	}));
 
 	{
-		FEditorMenuSection& Section = Menu->AddSection("WindowGlobalTabSpawners");
+		FToolMenuSection& Section = Menu->AddSection("WindowGlobalTabSpawners");
 		Section.AddMenuEntry(
 			"ProjectLauncher",
 			LOCTEXT("ProjectLauncherLabel", "Project Launcher"),
@@ -249,7 +249,7 @@ void FMainMenu::RegisterWindowMenu()
 		// Make sure at least one is enabled before creating the section
 		if (bLocalizationDashboard || bTranslationPicker)
 		{
-			FEditorMenuSection& Section = Menu->AddSection("ExperimentalTabSpawners", LOCTEXT("ExperimentalTabSpawnersHeading", "Experimental"));
+			FToolMenuSection& Section = Menu->AddSection("ExperimentalTabSpawners", LOCTEXT("ExperimentalTabSpawnersHeading", "Experimental"));
 			{
 				// Localization Dashboard
 				if (bLocalizationDashboard)
@@ -279,7 +279,7 @@ void FMainMenu::RegisterWindowMenu()
 	}
 
 	{
-		FEditorMenuSection& Section = Menu->AddSection("WindowLayout", NSLOCTEXT("MainAppMenu", "LayoutManagementHeader", "Layout"));
+		FToolMenuSection& Section = Menu->AddSection("WindowLayout", NSLOCTEXT("MainAppMenu", "LayoutManagementHeader", "Layout"));
 		Section.AddMenuEntry(FMainFrameCommands::Get().ResetLayout);
 		Section.AddMenuEntry(FMainFrameCommands::Get().SaveLayout);
 #if !PLATFORM_MAC && !PLATFORM_LINUX // On Mac/Linux windowed fullscreen mode in the editor is currently unavailable
@@ -295,14 +295,14 @@ void FMainMenu::RegisterWindowMenu()
 
 void FMainMenu::RegisterHelpMenu()
 {
-	UEditorMenu* Menu = UEditorMenuSubsystem::Get()->RegisterMenu("MainFrame.MainMenu.Help");
-	FEditorMenuSection& BugReportingSection = Menu->AddSection("BugReporting", NSLOCTEXT("MainHelpMenu", "BugsReporting", "Bugs"));
+	UToolMenu* Menu = UToolMenus::Get()->RegisterMenu("MainFrame.MainMenu.Help");
+	FToolMenuSection& BugReportingSection = Menu->AddSection("BugReporting", NSLOCTEXT("MainHelpMenu", "BugsReporting", "Bugs"));
 	{
 		BugReportingSection.AddMenuEntry(FMainFrameCommands::Get().ReportABug);
 		BugReportingSection.AddMenuEntry(FMainFrameCommands::Get().OpenIssueTracker);
 	}
 
-	FEditorMenuSection& HelpOnlineSection = Menu->AddSection("HelpOnline", NSLOCTEXT("MainHelpMenu", "Online", "Help Online"));
+	FToolMenuSection& HelpOnlineSection = Menu->AddSection("HelpOnline", NSLOCTEXT("MainHelpMenu", "Online", "Help Online"));
 	{
 		HelpOnlineSection.AddMenuEntry(FMainFrameCommands::Get().VisitSupportWebSite);
 		HelpOnlineSection.AddMenuEntry(FMainFrameCommands::Get().VisitForums);
@@ -320,7 +320,7 @@ void FMainMenu::RegisterHelpMenu()
 	}
 
 #if !PLATFORM_MAC // Handled by app's menu in menu bar
-	FEditorMenuSection& HelpApplicationSection = Menu->AddSection("HelpApplication", NSLOCTEXT("MainHelpMenu", "Application", "Application"));
+	FToolMenuSection& HelpApplicationSection = Menu->AddSection("HelpApplication", NSLOCTEXT("MainHelpMenu", "Application", "Application"));
 	{
 		const FText AboutWindowTitle = NSLOCTEXT("MainHelpMenu", "AboutUnrealEditor", "About Unreal Editor...");
 
@@ -331,21 +331,21 @@ void FMainMenu::RegisterHelpMenu()
 
 #undef LOCTEXT_NAMESPACE
 
-TSharedRef<SWidget> FMainMenu::MakeMainMenu(const TSharedPtr<FTabManager>& TabManager, const FName MenuName, FEditorMenuContext& EditorMenuContext)
+TSharedRef<SWidget> FMainMenu::MakeMainMenu(const TSharedPtr<FTabManager>& TabManager, const FName MenuName, FToolMenuContext& ToolMenuContext)
 {
 	// Cache all project names once
 	FMainFrameActionCallbacks::CacheProjectNames();
 
 	FMainMenu::RegisterMainMenu();
 
-	EditorMenuContext.AppendCommandList(FMainFrameCommands::ActionList);
+	ToolMenuContext.AppendCommandList(FMainFrameCommands::ActionList);
 
 	USlateTabManagerContext* ContextObject = NewObject<USlateTabManagerContext>();
 	ContextObject->TabManager = TabManager;
-	EditorMenuContext.AddObject(ContextObject);
+	ToolMenuContext.AddObject(ContextObject);
 
 	// Create the menu bar!
-	TSharedRef<SWidget> MenuBarWidget = UEditorMenuSubsystem::Get()->GenerateWidget(MenuName, EditorMenuContext);
+	TSharedRef<SWidget> MenuBarWidget = UToolMenus::Get()->GenerateWidget(MenuName, ToolMenuContext);
 	if (MenuBarWidget != SNullWidget::NullWidget)
 	{
 		// Tell tab-manager about the multi-box for platforms with a global menu bar
@@ -361,8 +361,8 @@ void FMainMenu::RegisterMainMenu()
 #define LOCTEXT_NAMESPACE "MainMenu"
 
 	static const FName MainMenuName("MainFrame.MainMenu");
-	UEditorMenuSubsystem* EditorMenus = UEditorMenuSubsystem::Get();
-	if (EditorMenus->IsMenuRegistered(MainMenuName))
+	UToolMenus* ToolMenus = UToolMenus::Get();
+	if (ToolMenus->IsMenuRegistered(MainMenuName))
 	{
 		return;
 	}
@@ -372,7 +372,7 @@ void FMainMenu::RegisterMainMenu()
 	RegisterWindowMenu();
 	RegisterHelpMenu();
 
-	UEditorMenu* MenuBar = EditorMenus->RegisterMenu(MainMenuName, NAME_None, EMultiBoxType::MenuBar);
+	UToolMenu* MenuBar = ToolMenus->RegisterMenu(MainMenuName, NAME_None, EMultiBoxType::MenuBar);
 
 	MenuBar->AddSubMenu(
 		"MainMenu",
@@ -419,9 +419,9 @@ void FMainMenu::RegisterFileProjectMenu()
 		return;
 	}
 
-	UEditorMenuSubsystem* EditorMenus = UEditorMenuSubsystem::Get();
-	UEditorMenu* MainTabFileMenu = EditorMenus->ExtendMenu("MainFrame.MainTabMenu.File");
-	FEditorMenuSection& Section = MainTabFileMenu->AddSection("FileProject", LOCTEXT("ProjectHeading", "Project"), FEditorMenuInsert("FileLoadAndSave", EEditorMenuInsertType::After));
+	UToolMenus* ToolMenus = UToolMenus::Get();
+	UToolMenu* MainTabFileMenu = ToolMenus->ExtendMenu("MainFrame.MainTabMenu.File");
+	FToolMenuSection& Section = MainTabFileMenu->AddSection("FileProject", LOCTEXT("ProjectHeading", "Project"), FToolMenuInsert("FileLoadAndSave", EToolMenuInsertType::After));
 
 	Section.AddMenuEntry( FMainFrameCommands::Get().NewProject );
 	Section.AddMenuEntry( FMainFrameCommands::Get().OpenProject );
@@ -433,7 +433,7 @@ void FMainMenu::RegisterFileProjectMenu()
 		FText::Format(LOCTEXT("AddCodeToProjectTooltip", "Adds C++ code to the project. The code can only be compiled if you have {0} installed."), ShortIDEName)
 	);
 
-	Section.AddEntry(FEditorMenuEntry::InitSubMenu(
+	Section.AddEntry(FToolMenuEntry::InitSubMenu(
 		"MainFrame.MainTabMenu.File",
 		"PackageProject",
 		LOCTEXT("PackageProjectSubMenuLabel", "Package Project"),
@@ -477,7 +477,7 @@ void FMainMenu::RegisterFileProjectMenu()
 		);
 	}
 
-	Section.AddDynamicEntry("CookContentForPlatform", FNewEditorMenuSectionDelegate::CreateLambda([](FEditorMenuSection& InSection)
+	Section.AddDynamicEntry("CookContentForPlatform", FNewToolMenuSectionDelegate::CreateLambda([](FToolMenuSection& InSection)
 	{
 		// @hack GDC: this should be moved somewhere else and be less hacky
 		ITargetPlatform* RunningTargetPlatform = GetTargetPlatformManager()->GetRunningTargetPlatform();
@@ -505,19 +505,19 @@ void FMainMenu::RegisterFileProjectMenu()
 
 void FMainMenu::RegisterRecentFileAndExitMenuItems()
 {
-	UEditorMenuSubsystem* EditorMenus = UEditorMenuSubsystem::Get();
-	UEditorMenu* MainTabFileMenu = EditorMenus->RegisterMenu("MainFrame.MainTabMenu.File", "MainFrame.MainMenu.File");
+	UToolMenus* ToolMenus = UToolMenus::Get();
+	UToolMenu* MainTabFileMenu = ToolMenus->RegisterMenu("MainFrame.MainTabMenu.File", "MainFrame.MainMenu.File");
 
 	{
-		FEditorMenuSection& Section = MainTabFileMenu->AddSection("FileRecentFiles");
+		FToolMenuSection& Section = MainTabFileMenu->AddSection("FileRecentFiles");
 		if (GetDefault<UEditorStyleSettings>()->bShowProjectMenus && FMainFrameActionCallbacks::ProjectNames.Num() > 0)
 		{
-			Section.AddEntry(FEditorMenuEntry::InitSubMenu(
+			Section.AddEntry(FToolMenuEntry::InitSubMenu(
 				"MainFrame.MainTabMenu.File",
 				"RecentProjects",
 				LOCTEXT("SwitchProjectSubMenu", "Recent Projects"),
 				LOCTEXT("SwitchProjectSubMenu_ToolTip", "Select a project to switch to"),
-				FNewEditorMenuDelegate::CreateStatic(&FRecentProjectsMenu::MakeMenu),
+				FNewToolMenuDelegate::CreateStatic(&FRecentProjectsMenu::MakeMenu),
 				false,
 				FSlateIcon(FEditorStyle::GetStyleSetName(), "MainFrame.RecentProjects")
 			));
@@ -526,15 +526,15 @@ void FMainMenu::RegisterRecentFileAndExitMenuItems()
 
 #if !PLATFORM_MAC // Handled by app's menu in menu bar
 	{
-		FEditorMenuSection& Section = MainTabFileMenu->AddSection("Exit");
+		FToolMenuSection& Section = MainTabFileMenu->AddSection("Exit");
 		Section.AddMenuEntry( FMainFrameCommands::Get().Exit );
 	}
 #endif
 }
 
-TSharedRef< SWidget > FMainMenu::MakeMainTabMenu( const TSharedPtr<FTabManager>& TabManager, const FName MenuName, FEditorMenuContext& EditorMenuContext )
+TSharedRef< SWidget > FMainMenu::MakeMainTabMenu( const TSharedPtr<FTabManager>& TabManager, const FName MenuName, FToolMenuContext& ToolMenuContext )
 {
-	return FMainMenu::MakeMainMenu( TabManager, MenuName, EditorMenuContext );
+	return FMainMenu::MakeMainMenu( TabManager, MenuName, ToolMenuContext );
 }
 
 
diff --git a/Engine/Source/Editor/MainFrame/Private/Menus/MainMenu.h b/Engine/Source/Editor/MainFrame/Private/Menus/MainMenu.h
--- a/Engine/Source/Editor/MainFrame/Private/Menus/MainMenu.h
+++ b/Engine/Source/Editor/MainFrame/Private/Menus/MainMenu.h
@@ -11,7 +11,7 @@
 #include "ILocalizationDashboardModule.h"
 
 class FMenuBuilder;
-struct FEditorMenuContext;
+struct FToolMenuContext;
 
 /**
  * Unreal editor main frame Slate widget
@@ -25,9 +25,9 @@ public:
 	 *
 	 * @param TabManager Create the workspace menu based on this tab manager.
 	 * @param MenuName Identifier associated with the menu.
-	 * @return EditorMenuContext Context containing state.
+	 * @return ToolMenuContext Context containing state.
 	 */
-	static TSharedRef<SWidget> MakeMainMenu( const TSharedPtr<FTabManager>& TabManager, const FName MenuName, FEditorMenuContext& EditorMenuContext );
+	static TSharedRef<SWidget> MakeMainMenu( const TSharedPtr<FTabManager>& TabManager, const FName MenuName, FToolMenuContext& ToolMenuContext );
 
 	/**
 	 * Static: Creates a widget for the main tab's menu bar.  This is just like the main menu bar, but also includes.
@@ -35,9 +35,9 @@ public:
 	 *
 	 * @param TabManager Create the workspace menu based on this tab manager.
 	 * @param MenuName Identifier associated with the menu.
-	 * @return EditorMenuContext Context containing state.
+	 * @return ToolMenuContext Context containing state.
 	 */
-	static TSharedRef<SWidget> MakeMainTabMenu( const TSharedPtr<FTabManager>& TabManager, const FName MenuName, FEditorMenuContext& EditorMenuContext );
+	static TSharedRef<SWidget> MakeMainTabMenu( const TSharedPtr<FTabManager>& TabManager, const FName MenuName, FToolMenuContext& ToolMenuContext );
 
 	/**
 	 * Static: Registers main menu with menu system.
diff --git a/Engine/Source/Editor/MainFrame/Private/Menus/RecentProjectsMenu.h b/Engine/Source/Editor/MainFrame/Private/Menus/RecentProjectsMenu.h
--- a/Engine/Source/Editor/MainFrame/Private/Menus/RecentProjectsMenu.h
+++ b/Engine/Source/Editor/MainFrame/Private/Menus/RecentProjectsMenu.h
@@ -4,7 +4,7 @@
 
 #include "CoreMinimal.h"
 #include "Misc/Paths.h"
-#include "EditorMenuSubsystem.h"
+#include "ToolMenus.h"
 #include "Frame/MainFrameActions.h"
 #include "HAL/FileManager.h"
 
@@ -22,7 +22,7 @@ public:
 	 *
 	 * @param Menu	The menu being populated.
 	 */
-	static void MakeMenu( UEditorMenu* Menu )
+	static void MakeMenu( UToolMenu* Menu )
 	{
 		for ( int32 ProjectIndex = 0; ProjectIndex < FMainFrameActionCallbacks::ProjectNames.Num() && ProjectIndex < FMainFrameCommands::Get().SwitchProjectCommands.Num(); ++ProjectIndex )
 		{
@@ -36,7 +36,7 @@ public:
 				continue;
 			}
 
-			FEditorMenuSection& Section = Menu->FindOrAddSection("Recent");
+			FToolMenuSection& Section = Menu->FindOrAddSection("Recent");
 			const FText DisplayName = FText::FromString( FPaths::GetBaseFilename(*ProjectName) );
 			const FText Tooltip = FText::FromString( IFileManager::Get().ConvertToAbsolutePathForExternalAppForRead(*ProjectName) );
 			Section.AddMenuEntry( FMainFrameCommands::Get().SwitchProjectCommands[ ProjectIndex ], DisplayName, Tooltip );
diff --git a/Engine/Source/Editor/MainFrame/Private/Menus/SettingsMenu.h b/Engine/Source/Editor/MainFrame/Private/Menus/SettingsMenu.h
--- a/Engine/Source/Editor/MainFrame/Private/Menus/SettingsMenu.h
+++ b/Engine/Source/Editor/MainFrame/Private/Menus/SettingsMenu.h
@@ -12,7 +12,7 @@
 #include "ISettingsContainer.h"
 #include "ISettingsSection.h"
 
-#include "EditorMenuSubsystem.h"
+#include "ToolMenus.h"
 
 #define LOCTEXT_NAMESPACE "FSettingsMenu"
 
@@ -29,7 +29,7 @@ public:
 	 * @param MenuBuilder The builder for the menu that owns this menu.
 	 * @param SettingsContainerName The name of the settings container to create the menu for.
 	 */
-	static void MakeMenu( UEditorMenu* MenuBuilder, FName SettingsContainerName )
+	static void MakeMenu( UToolMenu* MenuBuilder, FName SettingsContainerName )
 	{
 		ISettingsModule* SettingsModule = FModuleManager::GetModulePtr<ISettingsModule>("Settings");
 
@@ -59,7 +59,7 @@ public:
 
 			if (SettingsCategory->GetSections(SettingsSections) > 0)
 			{
-				FEditorMenuSection& Section = MenuBuilder->AddSection(SettingsCategory->GetName(), SettingsCategory->GetDisplayName());
+				FToolMenuSection& Section = MenuBuilder->AddSection(SettingsCategory->GetName(), SettingsCategory->GetDisplayName());
 
 				SettingsSections.Sort(
 					[](const ISettingsSectionPtr& First, const ISettingsSectionPtr& Second)
diff --git a/Engine/Source/Editor/MainFrame/Public/Interfaces/IMainFrameModule.h b/Engine/Source/Editor/MainFrame/Public/Interfaces/IMainFrameModule.h
--- a/Engine/Source/Editor/MainFrame/Public/Interfaces/IMainFrameModule.h
+++ b/Engine/Source/Editor/MainFrame/Public/Interfaces/IMainFrameModule.h
@@ -12,7 +12,7 @@
 
 class FTabManager;
 class SWindow;
-struct FEditorMenuContext;
+struct FToolMenuContext;
 
 /**
  * Data needed to display a Developer tools in the status bar.
@@ -53,7 +53,7 @@ public:
 	 *
 	 * @return	The newly-created menu widget
 	 */
-	virtual TSharedRef<SWidget> MakeMainMenu( const TSharedPtr<FTabManager>& TabManager, const FName MenuName, FEditorMenuContext& EditorMenuContext ) const = 0;
+	virtual TSharedRef<SWidget> MakeMainMenu( const TSharedPtr<FTabManager>& TabManager, const FName MenuName, FToolMenuContext& ToolMenuContext ) const = 0;
 
 	/**
 	 * Generates a menu that's just like the "main menu" widget above, except it also includes some infrequently used commands
@@ -64,7 +64,7 @@ public:
 	 *
 	 * @return	The newly-created menu widget
 	 */
-	virtual TSharedRef<SWidget> MakeMainTabMenu( const TSharedPtr<FTabManager>& TabManager, const FName MenuName, FEditorMenuContext& EditorMenuContext ) const = 0;
+	virtual TSharedRef<SWidget> MakeMainTabMenu( const TSharedPtr<FTabManager>& TabManager, const FName MenuName, FToolMenuContext& ToolMenuContext ) const = 0;
 
 	/**
 	 * Generates a menu for status and developer
diff --git a/Engine/Source/Editor/MaterialEditor/Private/MaterialInstanceEditor.cpp b/Engine/Source/Editor/MaterialEditor/Private/MaterialInstanceEditor.cpp
--- a/Engine/Source/Editor/MaterialEditor/Private/MaterialInstanceEditor.cpp
+++ b/Engine/Source/Editor/MaterialEditor/Private/MaterialInstanceEditor.cpp
@@ -19,7 +19,7 @@
 #include "Materials/MaterialFunctionInstance.h"
 #include "Materials/MaterialExpressionFunctionOutput.h"
 #include "MaterialEditorModule.h"
-#include "EditorMenuSubsystem.h"
+#include "ToolMenus.h"
 #include "Toolkits/AssetEditorToolkitMenuContext.h"
 
 #include "Materials/MaterialExpressionTextureBase.h"
@@ -778,28 +778,28 @@ void FMaterialInstanceEditor::UpdatePreviewViewportsVisibility()
 
 void FMaterialInstanceEditor::RegisterToolBar()
 {
-	UEditorMenuSubsystem* EditorMenus = UEditorMenuSubsystem::Get();
-	UEditorMenu* ToolBar = EditorMenus->ExtendMenu(GetEditorMenuToolbarName());
+	UToolMenus* ToolMenus = UToolMenus::Get();
+	UToolMenu* ToolBar = ToolMenus->ExtendMenu(GetToolMenuToolbarName());
 
-	FEditorMenuInsert InsertAfterAssetSection("Asset", EEditorMenuInsertType::After);
+	FToolMenuInsert InsertAfterAssetSection("Asset", EToolMenuInsertType::After);
 	{
-		FEditorMenuSection& Section = ToolBar->AddSection("Apply", TAttribute<FText>(), InsertAfterAssetSection);
-		Section.AddEntry(FEditorMenuEntry::InitToolBarButton(FMaterialEditorCommands::Get().Apply));
+		FToolMenuSection& Section = ToolBar->AddSection("Apply", TAttribute<FText>(), InsertAfterAssetSection);
+		Section.AddEntry(FToolMenuEntry::InitToolBarButton(FMaterialEditorCommands::Get().Apply));
 	}
 
 	{
-		FEditorMenuSection& Section = ToolBar->AddSection("Command", TAttribute<FText>(), InsertAfterAssetSection);
-		Section.AddEntry(FEditorMenuEntry::InitToolBarButton(FMaterialEditorCommands::Get().ShowAllMaterialParameters));
+		FToolMenuSection& Section = ToolBar->AddSection("Command", TAttribute<FText>(), InsertAfterAssetSection);
+		Section.AddEntry(FToolMenuEntry::InitToolBarButton(FMaterialEditorCommands::Get().ShowAllMaterialParameters));
 		// TODO: support in material instance editor.
-		Section.AddEntry(FEditorMenuEntry::InitToolBarButton(FMaterialEditorCommands::Get().TogglePlatformStats));
+		Section.AddEntry(FToolMenuEntry::InitToolBarButton(FMaterialEditorCommands::Get().TogglePlatformStats));
 	}
 	
 	{
-		FEditorMenuSection& Section = ToolBar->AddSection("Parent", TAttribute<FText>(), InsertAfterAssetSection);
-		Section.AddEntry(FEditorMenuEntry::InitComboButton(
+		FToolMenuSection& Section = ToolBar->AddSection("Parent", TAttribute<FText>(), InsertAfterAssetSection);
+		Section.AddEntry(FToolMenuEntry::InitComboButton(
 			"Hierarchy",
-			FEditorUIActionChoice(),
-			FNewEditorMenuDelegate::CreateSP(this, &FMaterialInstanceEditor::GenerateInheritanceMenu),
+			FToolUIActionChoice(),
+			FNewToolMenuDelegate::CreateSP(this, &FMaterialInstanceEditor::GenerateInheritanceMenu),
 			LOCTEXT("Hierarchy", "Hierarchy"),
 			FText::GetEmpty(),
 			FSlateIcon(FEditorStyle::GetStyleSetName(), "BTEditor.SwitchToBehaviorTreeMode"),
@@ -819,11 +819,11 @@ void FMaterialInstanceEditor::ExtendToolbar()
 	AddToolbarExtender(MaterialEditorModule->GetToolBarExtensibilityManager()->GetAllExtenders(GetToolkitCommands(), GetEditingObjects()));
 }
 
-void FMaterialInstanceEditor::GenerateInheritanceMenu(UEditorMenu* Menu)
+void FMaterialInstanceEditor::GenerateInheritanceMenu(UToolMenu* Menu)
 {
 	struct Local
 	{
-		static void AddMenuEntry(FEditorMenuSection& Section, FAssetData AssetData, bool bIsFunctionPreviewMaterial)
+		static void AddMenuEntry(FToolMenuSection& Section, FAssetData AssetData, bool bIsFunctionPreviewMaterial)
 		{
 			FExecuteAction OpenAction;
 			FExecuteAction FindInContentBrowserAction;
@@ -905,7 +905,7 @@ void FMaterialInstanceEditor::GenerateInheritanceMenu(UEditorMenu* Menu)
 					]
 				];
 
-			Section.AddEntry(FEditorMenuEntry::InitMenuEntry(
+			Section.AddEntry(FToolMenuEntry::InitMenuEntry(
 				NAME_None,
 				FUIAction(OpenAction),
 				EntryWidget
@@ -917,7 +917,7 @@ void FMaterialInstanceEditor::GenerateInheritanceMenu(UEditorMenu* Menu)
 	Menu->bShouldCloseWindowAfterMenuSelection = true;
 	Menu->bSearchable = true;
 	const FName ParentName = TEXT("ParentChain");
-	FEditorMenuSection& Section = Menu->AddSection(ParentName, LOCTEXT("ParentChain", "Parent Chain"));
+	FToolMenuSection& Section = Menu->AddSection(ParentName, LOCTEXT("ParentChain", "Parent Chain"));
 	if (bIsFunctionPreviewMaterial)
 	{
 		if (FunctionParentList.Num() == 0)
@@ -925,7 +925,7 @@ void FMaterialInstanceEditor::GenerateInheritanceMenu(UEditorMenu* Menu)
 			const FText NoParentText = LOCTEXT("NoParentFound", "No Parent Found");
 			TSharedRef<SWidget> NoParentWidget = SNew(STextBlock)
 				.Text(NoParentText);
-			Section.AddEntry(FEditorMenuEntry::InitWidget("NoParentEntry", NoParentWidget, FText::GetEmpty()));
+			Section.AddEntry(FToolMenuEntry::InitWidget("NoParentEntry", NoParentWidget, FText::GetEmpty()));
 		}
 		for (FAssetData FunctionParent : FunctionParentList)
 		{
@@ -939,7 +939,7 @@ void FMaterialInstanceEditor::GenerateInheritanceMenu(UEditorMenu* Menu)
 			const FText NoParentText = LOCTEXT("NoParentFound", "No Parent Found");
 			TSharedRef<SWidget> NoParentWidget = SNew(STextBlock)
 				.Text(NoParentText);
-			Section.AddEntry(FEditorMenuEntry::InitWidget("NoParentEntry", NoParentWidget, FText::GetEmpty()));
+			Section.AddEntry(FToolMenuEntry::InitWidget("NoParentEntry", NoParentWidget, FText::GetEmpty()));
 		}
 		for (FAssetData MaterialParent : MaterialParentList)
 		{
@@ -950,7 +950,7 @@ void FMaterialInstanceEditor::GenerateInheritanceMenu(UEditorMenu* Menu)
 	if (!bIsFunctionPreviewMaterial)
 	{
 		const FName MaterialInstances = TEXT("MaterialInstances");
-		FEditorMenuSection& MaterialInstancesSection = Menu->AddSection(MaterialInstances, LOCTEXT("MaterialInstances", "Material Instances"));
+		FToolMenuSection& MaterialInstancesSection = Menu->AddSection(MaterialInstances, LOCTEXT("MaterialInstances", "Material Instances"));
 		for (FAssetData MaterialChild : MaterialChildList)
 		{
 			Local::AddMenuEntry(MaterialInstancesSection, MaterialChild, bIsFunctionPreviewMaterial);
diff --git a/Engine/Source/Editor/MaterialEditor/Private/MaterialInstanceEditor.h b/Engine/Source/Editor/MaterialEditor/Private/MaterialInstanceEditor.h
--- a/Engine/Source/Editor/MaterialEditor/Private/MaterialInstanceEditor.h
+++ b/Engine/Source/Editor/MaterialEditor/Private/MaterialInstanceEditor.h
@@ -15,7 +15,7 @@
 #include "TickableEditorObject.h"
 
 class FCanvas;
-class UEditorMenu;
+class UToolMenu;
 class UMaterialEditorInstanceConstant;
 class UMaterialInterface;
 class UMaterialInstanceConstant;
@@ -158,7 +158,7 @@ private:
 	/** Builds the toolbar widget for the material editor */
 	void ExtendToolbar();
 
-	void GenerateInheritanceMenu(UEditorMenu* Menu);
+	void GenerateInheritanceMenu(UToolMenu* Menu);
 
 	/** If re-initializing for a material function instance re-generate the proxy materials */
 	void ReInitMaterialFunctionProxies();
diff --git a/Engine/Source/Editor/SceneOutliner/Private/FolderTreeItem.cpp b/Engine/Source/Editor/SceneOutliner/Private/FolderTreeItem.cpp
--- a/Engine/Source/Editor/SceneOutliner/Private/FolderTreeItem.cpp
+++ b/Engine/Source/Editor/SceneOutliner/Private/FolderTreeItem.cpp
@@ -3,7 +3,7 @@
 #include "FolderTreeItem.h"
 #include "Textures/SlateIcon.h"
 #include "Framework/Commands/UIAction.h"
-#include "EditorMenuSubsystem.h"
+#include "ToolMenus.h"
 #include "EditorStyleSet.h"
 #include "ScopedTransaction.h"
 #include "DragAndDrop/ActorDragDropGraphEdOp.h"
@@ -320,13 +320,13 @@ void FFolderTreeItem::OnExpansionChanged()
 	}
 }
 
-void FFolderTreeItem::GenerateContextMenu(UEditorMenu* Menu, SSceneOutliner& Outliner)
+void FFolderTreeItem::GenerateContextMenu(UToolMenu* Menu, SSceneOutliner& Outliner)
 {
 	auto SharedOutliner = StaticCastSharedRef<SSceneOutliner>(Outliner.AsShared());
 
 	const FSlateIcon NewFolderIcon(FEditorStyle::GetStyleSetName(), "SceneOutliner.NewFolderIcon");
 	
-	FEditorMenuSection& Section = Menu->AddSection("Section");
+	FToolMenuSection& Section = Menu->AddSection("Section");
 	Section.AddMenuEntry("CreateSubFolder", LOCTEXT("CreateSubFolder", "Create Sub Folder"), FText(), NewFolderIcon, FUIAction(FExecuteAction::CreateSP(this, &FFolderTreeItem::CreateSubFolder, TWeakPtr<SSceneOutliner>(SharedOutliner))));
 	Section.AddMenuEntry("DuplicateFolderHierarchy", LOCTEXT("DuplicateFolderHierarchy", "Duplicate Hierarchy"), FText(), FSlateIcon(), FUIAction(FExecuteAction::CreateSP(&Outliner, &SSceneOutliner::DuplicateFoldersHierarchy)));
 }
diff --git a/Engine/Source/Editor/SceneOutliner/Private/SSceneOutliner.cpp b/Engine/Source/Editor/SceneOutliner/Private/SSceneOutliner.cpp
--- a/Engine/Source/Editor/SceneOutliner/Private/SSceneOutliner.cpp
+++ b/Engine/Source/Editor/SceneOutliner/Private/SSceneOutliner.cpp
@@ -12,7 +12,7 @@
 #include "Widgets/Layout/SSeparator.h"
 #include "Widgets/Images/SImage.h"
 #include "Framework/MultiBox/MultiBoxBuilder.h"
-#include "EditorMenuSubsystem.h"
+#include "ToolMenus.h"
 #include "Widgets/Input/SButton.h"
 #include "Widgets/Input/SComboButton.h"
 #include "EditorStyleSet.h"
@@ -1716,7 +1716,7 @@ namespace SceneOutliner
 		ContextObject->NumSelectedItems = OutlinerTreeView->GetNumItemsSelected();
 		ContextObject->NumSelectedFolders = ItemSelection.Folders.Num();
 		ContextObject->NumWorldsSelected = ItemSelection.Worlds.Num();
-		FEditorMenuContext Context(ContextObject);
+		FToolMenuContext Context(ContextObject);
 
 		// Allow other systems to override menu name and provide additional context
 		static const FName DefaultContextMenuName("SceneOutliner.DefaultContextMenu");
@@ -1724,14 +1724,14 @@ namespace SceneOutliner
 		SharedData->ModifyContextMenu.ExecuteIfBound(MenuName, Context);
 
 		// Build up the menu for a selection
-		UEditorMenuSubsystem* EditorMenus = UEditorMenuSubsystem::Get();
-		UEditorMenu* Menu = EditorMenus->GenerateMenu(MenuName, Context);
+		UToolMenus* ToolMenus = UToolMenus::Get();
+		UToolMenu* Menu = ToolMenus->GenerateMenu(MenuName, Context);
 
-		for (const FEditorMenuSection& Section : Menu->Sections)
+		for (const FToolMenuSection& Section : Menu->Sections)
 		{
 			if (Section.Blocks.Num() > 0)
 			{
-				return EditorMenus->GenerateWidget(Menu);
+				return ToolMenus->GenerateWidget(Menu);
 			}
 		}
 
@@ -1743,13 +1743,13 @@ namespace SceneOutliner
 		static const FName DefaultContextBaseMenuName("SceneOutliner.DefaultContextMenuBase");
 		static const FName DefaultContextMenuName("SceneOutliner.DefaultContextMenu");
 
-		UEditorMenuSubsystem* EditorMenus = UEditorMenuSubsystem::Get();
+		UToolMenus* ToolMenus = UToolMenus::Get();
 
-		if (!EditorMenus->IsMenuRegistered(DefaultContextBaseMenuName))
+		if (!ToolMenus->IsMenuRegistered(DefaultContextBaseMenuName))
 		{
-			UEditorMenu* Menu = EditorMenus->RegisterMenu(DefaultContextBaseMenuName);
+			UToolMenu* Menu = ToolMenus->RegisterMenu(DefaultContextBaseMenuName);
 
-			Menu->AddDynamicSection("DynamicSection1", FNewEditorMenuDelegate::CreateLambda([](UEditorMenu* InMenu)
+			Menu->AddDynamicSection("DynamicSection1", FNewToolMenuDelegate::CreateLambda([](UToolMenu* InMenu)
 			{
 				USceneOutlinerMenuContext* Context = InMenu->FindContext<USceneOutlinerMenuContext>();
 				if (!Context || !Context->SceneOutliner.IsValid())
@@ -1779,21 +1779,21 @@ namespace SceneOutliner
 						// If we've only got folders selected, show the selection and edit sub menus
 						if (Context->NumSelectedItems > 0 && Context->NumSelectedFolders == Context->NumSelectedItems)
 						{
-							InMenu->FindOrAddSection("Section").AddEntry(FEditorMenuEntry::InitSubMenu(
+							InMenu->FindOrAddSection("Section").AddEntry(FToolMenuEntry::InitSubMenu(
 								InMenu->MenuName,
 								"SelectSubMenu",
 								LOCTEXT("SelectSubmenu", "Select"),
 								LOCTEXT("SelectSubmenu_Tooltip", "Select the contents of the current selection"),
-								FNewEditorMenuDelegate::CreateSP(SceneOutliner, &SSceneOutliner::FillSelectionSubMenu)));
+								FNewToolMenuDelegate::CreateSP(SceneOutliner, &SSceneOutliner::FillSelectionSubMenu)));
 						}
 					}
 				}
 			}));
 
-			Menu->AddDynamicSection("DynamicMainSection", FNewEditorMenuDelegate::CreateLambda([](UEditorMenu* InMenu)
+			Menu->AddDynamicSection("DynamicMainSection", FNewToolMenuDelegate::CreateLambda([](UToolMenu* InMenu)
 			{
 				// We always create a section here, even if there is no parent so that clients can still extend the menu
-				FEditorMenuSection& Section = InMenu->AddSection("MainSection");
+				FToolMenuSection& Section = InMenu->AddSection("MainSection");
 
 				if (USceneOutlinerMenuContext* Context = InMenu->FindContext<USceneOutlinerMenuContext>())
 				{
@@ -1801,26 +1801,26 @@ namespace SceneOutliner
 					// Can't move worlds or level blueprints
 					if (Context->bShowParentTree && Context->NumSelectedItems > 0 && Context->NumWorldsSelected == 0 && Context->SceneOutliner.IsValid())
 					{
-						Section.AddEntry(FEditorMenuEntry::InitSubMenu(
+						Section.AddEntry(FToolMenuEntry::InitSubMenu(
 							InMenu->MenuName,
 							"MoveActorsTo",
 							LOCTEXT("MoveActorsTo", "Move To"),
 							LOCTEXT("MoveActorsTo_Tooltip", "Move selection to another folder"),
-							FNewEditorMenuDelegate::CreateSP(Context->SceneOutliner.Pin().Get(), &SSceneOutliner::FillFoldersSubMenu)));
+							FNewToolMenuDelegate::CreateSP(Context->SceneOutliner.Pin().Get(), &SSceneOutliner::FillFoldersSubMenu)));
 					}
 				}
 			}));
 		}
 
-		if (!EditorMenus->IsMenuRegistered(DefaultContextMenuName))
+		if (!ToolMenus->IsMenuRegistered(DefaultContextMenuName))
 		{
-			EditorMenus->RegisterMenu(DefaultContextMenuName, DefaultContextBaseMenuName);
+			ToolMenus->RegisterMenu(DefaultContextMenuName, DefaultContextBaseMenuName);
 		}
 	}
 
-	void SSceneOutliner::FillFoldersSubMenu(UEditorMenu* Menu) const
+	void SSceneOutliner::FillFoldersSubMenu(UToolMenu* Menu) const
 	{
-		FEditorMenuSection& Section = Menu->AddSection("Section");
+		FToolMenuSection& Section = Menu->AddSection("Section");
 		Section.AddMenuEntry("CreateNew", LOCTEXT( "CreateNew", "Create New Folder" ), LOCTEXT( "CreateNew_ToolTip", "Move to a new folder" ),
 			FSlateIcon(FEditorStyle::GetStyleSetName(), "SceneOutliner.NewFolderIcon"), FExecuteAction::CreateSP(const_cast<SSceneOutliner*>(this), &SSceneOutliner::CreateFolder));
 
@@ -1891,7 +1891,7 @@ namespace SceneOutliner
 		return ExcludedParents;
 	}
 
-	void SSceneOutliner::AddMoveToFolderOutliner(UEditorMenu* Menu) const
+	void SSceneOutliner::AddMoveToFolderOutliner(UToolMenu* Menu) const
 	{
 		// We don't show this if there aren't any folders in the world
 		if (!FActorFolders::Get().GetFolderPropertiesForWorld(*SharedData->RepresentingWorld).Num())
@@ -1951,17 +1951,17 @@ namespace SceneOutliner
 				.OnItemPickedDelegate(FOnSceneOutlinerItemPicked::CreateSP(const_cast<SSceneOutliner*>(this), &SSceneOutliner::MoveSelectionTo))
 			];
 
-		FEditorMenuSection& Section = Menu->AddSection(FName(), LOCTEXT("ExistingFolders", "Existing:"));
-		Section.AddEntry(FEditorMenuEntry::InitWidget(
+		FToolMenuSection& Section = Menu->AddSection(FName(), LOCTEXT("ExistingFolders", "Existing:"));
+		Section.AddEntry(FToolMenuEntry::InitWidget(
 			"MiniSceneOutliner",
 			MiniSceneOutliner,
 			FText::GetEmpty(),
 			false));
 	}
 
-	void SSceneOutliner::FillSelectionSubMenu(UEditorMenu* Menu) const
+	void SSceneOutliner::FillSelectionSubMenu(UToolMenu* Menu) const
 	{
-		FEditorMenuSection& Section = Menu->AddSection("Section");
+		FToolMenuSection& Section = Menu->AddSection("Section");
 		Section.AddMenuEntry(
 			"AddChildrenToSelection",
 			LOCTEXT( "AddChildrenToSelection", "Immediate Children" ),
diff --git a/Engine/Source/Editor/SceneOutliner/Private/WorldTreeItem.cpp b/Engine/Source/Editor/SceneOutliner/Private/WorldTreeItem.cpp
--- a/Engine/Source/Editor/SceneOutliner/Private/WorldTreeItem.cpp
+++ b/Engine/Source/Editor/SceneOutliner/Private/WorldTreeItem.cpp
@@ -9,7 +9,7 @@
 #include "SceneOutlinerDragDrop.h"
 #include "SSceneOutliner.h"
 
-#include "EditorMenuSubsystem.h"
+#include "ToolMenus.h"
 #include "LevelEditor.h"
 #include "EditorActorFolders.h"
 
@@ -82,13 +82,13 @@ bool FWorldTreeItem::CanInteract() const
 	return Flags.bInteractive;
 }
 
-void FWorldTreeItem::GenerateContextMenu(UEditorMenu* Menu, SSceneOutliner& Outliner)
+void FWorldTreeItem::GenerateContextMenu(UToolMenu* Menu, SSceneOutliner& Outliner)
 {
 	auto SharedOutliner = StaticCastSharedRef<SSceneOutliner>(Outliner.AsShared());
 	
 	const FSlateIcon WorldSettingsIcon(FEditorStyle::GetStyleSetName(), "LevelEditor.WorldProperties.Tab");
 	const FSlateIcon NewFolderIcon(FEditorStyle::GetStyleSetName(), "SceneOutliner.NewFolderIcon");
-	FEditorMenuSection& Section = Menu->AddSection("Section");
+	FToolMenuSection& Section = Menu->AddSection("Section");
 	Section.AddMenuEntry("CreateFolder", LOCTEXT("CreateFolder", "Create Folder"), FText(), NewFolderIcon, FUIAction(FExecuteAction::CreateSP(this, &FWorldTreeItem::CreateFolder, TWeakPtr<SSceneOutliner>(SharedOutliner))));
 	Section.AddMenuEntry("OpenWorldSettings", LOCTEXT("OpenWorldSettings", "World Settings"), FText(), WorldSettingsIcon, FExecuteAction::CreateSP(this, &FWorldTreeItem::OpenWorldSettings));
 }
diff --git a/Engine/Source/Editor/SceneOutliner/Public/FolderTreeItem.h b/Engine/Source/Editor/SceneOutliner/Public/FolderTreeItem.h
--- a/Engine/Source/Editor/SceneOutliner/Public/FolderTreeItem.h
+++ b/Engine/Source/Editor/SceneOutliner/Public/FolderTreeItem.h
@@ -6,7 +6,7 @@
 #include "Widgets/SWidget.h"
 #include "ITreeItem.h"
 
-class UEditorMenu;
+class UToolMenu;
 
 namespace SceneOutliner
 {
@@ -68,7 +68,7 @@ namespace SceneOutliner
 		virtual bool CanInteract() const override;
 		
 		/** Generate a context menu for this item. Only called if *only* this item is selected. */
-		virtual void GenerateContextMenu(UEditorMenu* Menu, SSceneOutliner& Outliner) override;
+		virtual void GenerateContextMenu(UToolMenu* Menu, SSceneOutliner& Outliner) override;
 
 		/** Populate the specified drag/drop payload with any relevant information for this type */
 		virtual void PopulateDragDropPayload(FDragDropPayload& Payload) const override;
diff --git a/Engine/Source/Editor/SceneOutliner/Public/ITreeItem.h b/Engine/Source/Editor/SceneOutliner/Public/ITreeItem.h
--- a/Engine/Source/Editor/SceneOutliner/Public/ITreeItem.h
+++ b/Engine/Source/Editor/SceneOutliner/Public/ITreeItem.h
@@ -6,7 +6,7 @@
 #include "Widgets/SWidget.h"
 #include "SceneOutlinerVisitorTypes.h"
 
-class UEditorMenu;
+class UToolMenu;
 
 namespace SceneOutliner
 {
@@ -156,7 +156,7 @@ namespace SceneOutliner
 		virtual void OnExpansionChanged() {};
 
 		/** Generate a context menu for this item. Only called if *only* this item is selected. */
-		virtual void GenerateContextMenu(UEditorMenu* Menu, SSceneOutliner& Outliner) {}
+		virtual void GenerateContextMenu(UToolMenu* Menu, SSceneOutliner& Outliner) {}
 
 		/** Populate the specified drag/drop payload with any relevant information for this type */
 		virtual void PopulateDragDropPayload(FDragDropPayload& Payload) const = 0;
diff --git a/Engine/Source/Editor/SceneOutliner/Public/SSceneOutliner.h b/Engine/Source/Editor/SceneOutliner/Public/SSceneOutliner.h
--- a/Engine/Source/Editor/SceneOutliner/Public/SSceneOutliner.h
+++ b/Engine/Source/Editor/SceneOutliner/Public/SSceneOutliner.h
@@ -20,7 +20,7 @@
 #include "Framework/Commands/UICommandList.h"
 
 class FMenuBuilder;
-class UEditorMenu;
+class UToolMenu;
 class ISceneOutlinerColumn;
 class SComboButton;
 
@@ -507,9 +507,9 @@ namespace SceneOutliner
 
 		/** Build a context menu for right-clicking an item in the tree */
 		TSharedPtr<SWidget> BuildDefaultContextMenu();
-		void FillFoldersSubMenu(UEditorMenu* Menu) const;
-		void AddMoveToFolderOutliner(UEditorMenu* Menu) const;
-		void FillSelectionSubMenu(UEditorMenu* Menun) const;
+		void FillFoldersSubMenu(UToolMenu* Menu) const;
+		void AddMoveToFolderOutliner(UToolMenu* Menu) const;
+		void FillSelectionSubMenu(UToolMenu* Menun) const;
 		TSharedRef<TSet<FName>> GatherInvalidMoveToDestinations() const;
 
 	private:
diff --git a/Engine/Source/Editor/SceneOutliner/Public/SceneOutlinerPublicTypes.h b/Engine/Source/Editor/SceneOutliner/Public/SceneOutlinerPublicTypes.h
--- a/Engine/Source/Editor/SceneOutliner/Public/SceneOutlinerPublicTypes.h
+++ b/Engine/Source/Editor/SceneOutliner/Public/SceneOutlinerPublicTypes.h
@@ -10,7 +10,7 @@
 #include "SceneOutlinerFilters.h"
 
 class FExtender;
-struct FEditorMenuContext;
+struct FToolMenuContext;
 
 enum class ESceneOutlinerMode : uint8
 {
@@ -24,7 +24,7 @@ enum class ESceneOutlinerMode : uint8
 	ComponentPicker,
 };
 
-DECLARE_DELEGATE_TwoParams(FSceneOutlinerModifyContextMenu, FName& /* MenuName */, FEditorMenuContext& /* MenuContext */);
+DECLARE_DELEGATE_TwoParams(FSceneOutlinerModifyContextMenu, FName& /* MenuName */, FToolMenuContext& /* MenuContext */);
 
 namespace SceneOutliner
 {
diff --git a/Engine/Source/Editor/SceneOutliner/Public/WorldTreeItem.h b/Engine/Source/Editor/SceneOutliner/Public/WorldTreeItem.h
--- a/Engine/Source/Editor/SceneOutliner/Public/WorldTreeItem.h
+++ b/Engine/Source/Editor/SceneOutliner/Public/WorldTreeItem.h
@@ -7,7 +7,7 @@
 #include "ITreeItem.h"
 #include "UObject/ObjectKey.h"
 
-class UEditorMenu;
+class UToolMenu;
 
 namespace SceneOutliner
 {
@@ -48,7 +48,7 @@ namespace SceneOutliner
 		virtual bool CanInteract() const override;
 
 		/** Generate a context menu for this item. Only called if *only* this item is selected. */
-		virtual void GenerateContextMenu(UEditorMenu* Menu, SSceneOutliner& Outliner) override;
+		virtual void GenerateContextMenu(UToolMenu* Menu, SSceneOutliner& Outliner) override;
 
 		/** Populate the specified drag/drop payload with any relevant information for this type */
 		virtual void PopulateDragDropPayload(FDragDropPayload& Payload) const override
diff --git a/Engine/Source/Editor/UnrealEd/Private/EditorEngine.cpp b/Engine/Source/Editor/UnrealEd/Private/EditorEngine.cpp
--- a/Engine/Source/Editor/UnrealEd/Private/EditorEngine.cpp
+++ b/Engine/Source/Editor/UnrealEd/Private/EditorEngine.cpp
@@ -213,6 +213,7 @@
 #include "ComponentRecreateRenderStateContext.h"
 #include "RenderTargetPool.h"
 #include "RenderGraphBuilder.h"
+#include "ToolMenus.h"
 
 DEFINE_LOG_CATEGORY_STATIC(LogEditor, Log, All);
 
@@ -779,6 +780,15 @@ void UEditorEngine::InitEditor(IEngineLoop* InEngineLoop)
 	};
 	// Sort by menu priority.
 	ActorFactories.Sort(FCompareUActorFactoryByMenuPriority());
+
+	if (FSlateApplication::IsInitialized() && UToolMenus::IsToolMenuUIEnabled())
+	{
+		UToolMenus::Get()->SetTimerManager(TimerManager.ToSharedRef());
+		UToolMenus::Get()->ShouldDisplayExtensionPoints.BindStatic(&GetDisplayMultiboxHooks);
+		UToolMenus::Get()->RegisterStringCommandHandler("Command", FToolMenuExecuteString::CreateLambda([](const FString& InString, const FToolMenuContext& InContext) {
+			GEditor->Exec(nullptr, *InString);
+		}));
+	}
 }
 
 bool UEditorEngine::HandleOpenAsset(UObject* Asset)
@@ -1172,6 +1182,12 @@ void UEditorEngine::FinishDestroy()
 			UE_LOG(LogEditor, Warning, TEXT("Warning: Play world is active"));
 		}
 
+		if (UToolMenus* ToolMenus = UToolMenus::TryGet())
+		{
+			ToolMenus->ShouldDisplayExtensionPoints.Unbind();
+			ToolMenus->UnregisterStringCommandHandler("Command");
+		}
+
 		EditorSubsystemCollection.Deinitialize();
 
 		// Unregister events
diff --git a/Engine/Source/Editor/UnrealEd/Private/Kismet2/ComponentEditorUtils.cpp b/Engine/Source/Editor/UnrealEd/Private/Kismet2/ComponentEditorUtils.cpp
--- a/Engine/Source/Editor/UnrealEd/Private/Kismet2/ComponentEditorUtils.cpp
+++ b/Engine/Source/Editor/UnrealEd/Private/Kismet2/ComponentEditorUtils.cpp
@@ -27,7 +27,7 @@
 #include "Editor.h"
 #include "Toolkits/AssetEditorManager.h"
 #include "HAL/PlatformApplicationMisc.h"
-#include "EditorMenuSubsystem.h"
+#include "ToolMenus.h"
 #include "Kismet2/ComponentEditorContextMenuContex.h"
 
 #define LOCTEXT_NAMESPACE "ComponentEditorUtils"
@@ -915,11 +915,11 @@ FName FComponentEditorUtils::FindVariableNameGivenComponentInstance(const UActor
 	return NAME_None;
 }
 
-void FComponentEditorUtils::FillComponentContextMenuOptions(UEditorMenu* Menu, const TArray<UActorComponent*>& SelectedComponents)
+void FComponentEditorUtils::FillComponentContextMenuOptions(UToolMenu* Menu, const TArray<UActorComponent*>& SelectedComponents)
 {
 	// Basic commands
 	{
-		FEditorMenuSection& Section = Menu->AddSection("EditComponent", LOCTEXT("EditComponentHeading", "Edit"));
+		FToolMenuSection& Section = Menu->AddSection("EditComponent", LOCTEXT("EditComponentHeading", "Edit"));
 		Section.AddMenuEntry(FGenericCommands::Get().Cut);
 		Section.AddMenuEntry(FGenericCommands::Get().Copy);
 		Section.AddMenuEntry(FGenericCommands::Get().Paste);
@@ -935,7 +935,7 @@ void FComponentEditorUtils::FillComponentContextMenuOptions(UEditorMenu* Menu, c
 		if (Component->GetClass()->ClassGeneratedBy)
 		{
 			{
-				FEditorMenuSection& Section = Menu->AddSection("ComponentAsset", LOCTEXT("ComponentAssetHeading", "Asset"));
+				FToolMenuSection& Section = Menu->AddSection("ComponentAsset", LOCTEXT("ComponentAssetHeading", "Asset"));
 				Section.AddMenuEntry(
 					"GoToBlueprintForComponent",
 					FText::Format(LOCTEXT("GoToBlueprintForComponent", "Edit {0}"), FText::FromString(Component->GetClass()->ClassGeneratedBy->GetName())),
@@ -958,7 +958,7 @@ void FComponentEditorUtils::FillComponentContextMenuOptions(UEditorMenu* Menu, c
 		else
 		{
 			{
-				FEditorMenuSection& Section = Menu->AddSection("ComponentCode", LOCTEXT("ComponentCodeHeading", "C++"));
+				FToolMenuSection& Section = Menu->AddSection("ComponentCode", LOCTEXT("ComponentCodeHeading", "C++"));
 				if (FSourceCodeNavigation::IsCompilerAvailable())
 				{
 					FString ClassHeaderPath;
diff --git a/Engine/Source/Editor/UnrealEd/Private/Kismet2/Kismet2.cpp b/Engine/Source/Editor/UnrealEd/Private/Kismet2/Kismet2.cpp
--- a/Engine/Source/Editor/UnrealEd/Private/Kismet2/Kismet2.cpp
+++ b/Engine/Source/Editor/UnrealEd/Private/Kismet2/Kismet2.cpp
@@ -82,7 +82,7 @@
 #include "Engine/InheritableComponentHandler.h"
 #include "Stats/StatsHierarchical.h"
 #include "Classes/EditorStyleSettings.h"
-#include "EditorMenuSubsystem.h"
+#include "ToolMenus.h"
 
 DECLARE_CYCLE_STAT(TEXT("Compile Blueprint"), EKismetCompilerStats_CompileBlueprint, STATGROUP_KismetCompiler);
 DECLARE_CYCLE_STAT(TEXT("Broadcast Precompile"), EKismetCompilerStats_BroadcastPrecompile, STATGROUP_KismetCompiler);
@@ -2452,7 +2452,7 @@ bool FKismetEditorUtilities::AnyBoundLevelScriptEventForActor(AActor* Actor, boo
 	return false;
 }
 
-void FKismetEditorUtilities::AddLevelScriptEventOptionsForActor(UEditorMenu* Menu, TWeakObjectPtr<AActor> ActorPtr, bool bExistingEvents, bool bNewEvents, bool bOnlyEventName)
+void FKismetEditorUtilities::AddLevelScriptEventOptionsForActor(UToolMenu* Menu, TWeakObjectPtr<AActor> ActorPtr, bool bExistingEvents, bool bNewEvents, bool bOnlyEventName)
 {
 	struct FCreateEventForActorHelper
 	{
@@ -2522,7 +2522,7 @@ void FKismetEditorUtilities::AddLevelScriptEventOptionsForActor(UEditorMenu* Men
 		// Now build the menu
 		for(FEventCategory& Category : CategorizedEvents)
 		{
-			FEditorMenuSection& Section = Menu->AddSection(NAME_None, FText::FromString(Category.CategoryName));
+			FToolMenuSection& Section = Menu->AddSection(NAME_None, FText::FromString(Category.CategoryName));
 
 			for(UProperty* Property : Category.EventProperties)
 			{
diff --git a/Engine/Source/Editor/UnrealEd/Private/Toolkits/AssetEditorToolkit.cpp b/Engine/Source/Editor/UnrealEd/Private/Toolkits/AssetEditorToolkit.cpp
--- a/Engine/Source/Editor/UnrealEd/Private/Toolkits/AssetEditorToolkit.cpp
+++ b/Engine/Source/Editor/UnrealEd/Private/Toolkits/AssetEditorToolkit.cpp
@@ -28,7 +28,7 @@
 #include "IAssetTypeActions.h"
 #include "AssetToolsModule.h"
 #include "Toolkits/AssetEditorToolkitMenuContext.h"
-#include "EditorMenuSubsystem.h"
+#include "ToolMenus.h"
 
 #define LOCTEXT_NAMESPACE "AssetEditorToolkit"
 
@@ -962,7 +962,7 @@ void FAssetEditorToolkit::FillDefaultHelpMenuCommands( FMenuBuilder& MenuBuilder
 	MenuBuilder.AddMenuEntry(FGlobalEditorCommonCommands::Get().OpenDocumentation, NAME_None, ToolTip);
 }
 
-FName FAssetEditorToolkit::GetEditorMenuAppName() const
+FName FAssetEditorToolkit::GetToolMenuAppName() const
 {
 	if (IsSimpleAssetEditor() && EditingObjects.Num() == 1 && EditingObjects[0])
 	{
@@ -972,22 +972,22 @@ FName FAssetEditorToolkit::GetEditorMenuAppName() const
 	return GetToolkitFName();
 }
 
-FName FAssetEditorToolkit::GetEditorMenuToolbarName() const
+FName FAssetEditorToolkit::GetToolMenuToolbarName() const
 {
-	return *(TEXT("AssetEditor.") + GetEditorMenuAppName().ToString() + TEXT(".ToolBar"));
+	return *(TEXT("AssetEditor.") + GetToolMenuAppName().ToString() + TEXT(".ToolBar"));
 }
 
 void FAssetEditorToolkit::RegisterMenus()
 {
 	static const FName DefaultToolBarName("AssetEditor.DefaultToolBar");
-	UEditorMenuSubsystem* EditorMenus = UEditorMenuSubsystem::Get();
-	if (!EditorMenus->IsMenuRegistered(DefaultToolBarName))
+	UToolMenus* ToolMenus = UToolMenus::Get();
+	if (!ToolMenus->IsMenuRegistered(DefaultToolBarName))
 	{
-		UEditorMenu* ToolbarBuilder = EditorMenus->RegisterMenu(DefaultToolBarName, NAME_None, EMultiBoxType::ToolBar);
+		UToolMenu* ToolbarBuilder = ToolMenus->RegisterMenu(DefaultToolBarName, NAME_None, EMultiBoxType::ToolBar);
 		{
-			FEditorMenuSection& Section = ToolbarBuilder->AddSection("Asset");
-			Section.AddEntry(FEditorMenuEntry::InitToolBarButton(FAssetEditorCommonCommands::Get().SaveAsset));
-			Section.AddEntry(FEditorMenuEntry::InitToolBarButton(FGlobalEditorCommonCommands::Get().FindInContentBrowser, LOCTEXT("FindInContentBrowserButton", "Browse")));
+			FToolMenuSection& Section = ToolbarBuilder->AddSection("Asset");
+			Section.AddEntry(FToolMenuEntry::InitToolBarButton(FAssetEditorCommonCommands::Get().SaveAsset));
+			Section.AddEntry(FToolMenuEntry::InitToolBarButton(FGlobalEditorCommonCommands::Get().FindInContentBrowser, LOCTEXT("FindInContentBrowserButton", "Browse")));
 		}
 	}
 }
@@ -998,24 +998,24 @@ void FAssetEditorToolkit::GenerateToolbar()
 
 	RegisterMenus();
 
-	const FName ToolBarName = GetEditorMenuToolbarName();
-	UEditorMenuSubsystem* EditorMenus = UEditorMenuSubsystem::Get();
-	UEditorMenu* FoundMenu = EditorMenus->FindMenu(ToolBarName);
+	const FName ToolBarName = GetToolMenuToolbarName();
+	UToolMenus* ToolMenus = UToolMenus::Get();
+	UToolMenu* FoundMenu = ToolMenus->FindMenu(ToolBarName);
 	if (!FoundMenu || !FoundMenu->IsRegistered())
 	{
-		FoundMenu = EditorMenus->RegisterMenu(ToolBarName, "AssetEditor.DefaultToolBar", EMultiBoxType::ToolBar);
+		FoundMenu = ToolMenus->RegisterMenu(ToolBarName, "AssetEditor.DefaultToolBar", EMultiBoxType::ToolBar);
 	}
 
-	FEditorMenuContext MenuContext(GetToolkitCommands(), Extender);
+	FToolMenuContext MenuContext(GetToolkitCommands(), Extender);
 
 	UAssetEditorToolkitMenuContext* ToolkitMenuContext = NewObject<UAssetEditorToolkitMenuContext>(FoundMenu);
 	ToolkitMenuContext->Toolkit = AsShared();
 	MenuContext.AddObject(ToolkitMenuContext);
 
-	UEditorMenu* GeneratedToolbar = EditorMenus->GenerateMenu(ToolBarName, MenuContext);
+	UToolMenu* GeneratedToolbar = ToolMenus->GenerateMenu(ToolBarName, MenuContext);
 	GeneratedToolbar->bToolBarIsFocusable = bIsToolbarFocusable;
 	GeneratedToolbar->bToolBarForceSmallIcons = bIsToolbarUsingSmallIcons;
-	TSharedRef< class SWidget > ToolBarWidget = EditorMenus->GenerateWidget(GeneratedToolbar);
+	TSharedRef< class SWidget > ToolBarWidget = ToolMenus->GenerateWidget(GeneratedToolbar);
 
 	TSharedRef<SHorizontalBox> MiscWidgets = SNew(SHorizontalBox);
 
diff --git a/Engine/Source/Editor/UnrealEd/Private/Toolkits/SStandaloneAssetEditorToolkitHost.cpp b/Engine/Source/Editor/UnrealEd/Private/Toolkits/SStandaloneAssetEditorToolkitHost.cpp
--- a/Engine/Source/Editor/UnrealEd/Private/Toolkits/SStandaloneAssetEditorToolkitHost.cpp
+++ b/Engine/Source/Editor/UnrealEd/Private/Toolkits/SStandaloneAssetEditorToolkitHost.cpp
@@ -12,7 +12,7 @@
 #include "Widgets/Docking/SDockTab.h"
 #include "UObject/Package.h"
 
-#include "EditorMenuSubsystem.h"
+#include "ToolMenus.h"
 
 #define LOCTEXT_NAMESPACE "StandaloneAssetEditorToolkit"
 
@@ -43,9 +43,9 @@ void SStandaloneAssetEditorToolkitHost::SetupInitialContent( const TSharedRef<FT
 	//				- In standalone, just draws under the toolkit's menu
 
 	const FName AssetEditorMenuName = GetMenuName();
-	if (!UEditorMenuSubsystem::Get()->IsMenuRegistered(AssetEditorMenuName))
+	if (!UToolMenus::Get()->IsMenuRegistered(AssetEditorMenuName))
 	{
-		UEditorMenuSubsystem::Get()->RegisterMenu(AssetEditorMenuName, "MainFrame.MainMenu");
+		UToolMenus::Get()->RegisterMenu(AssetEditorMenuName, "MainFrame.MainMenu");
 	}
 
 	if (bCreateDefaultStandaloneMenu)
@@ -184,7 +184,7 @@ FName SStandaloneAssetEditorToolkitHost::GetMenuName() const
 	FName MenuAppName;
 	if (HostedAssetEditorToolkit.IsValid())
 	{
-		MenuAppName = HostedAssetEditorToolkit->GetEditorMenuAppName();
+		MenuAppName = HostedAssetEditorToolkit->GetToolMenuAppName();
 	}
 	else
 	{
@@ -199,9 +199,9 @@ void SStandaloneAssetEditorToolkitHost::GenerateMenus(bool bForceCreateMenu)
 	if( bForceCreateMenu || DefaultMenuWidget != SNullWidget::NullWidget )
 	{
 		const FName AssetEditorMenuName = GetMenuName();
-		FEditorMenuContext EditorMenuContext(HostedAssetEditorToolkit->GetToolkitCommands(), FExtender::Combine(MenuExtenders).ToSharedRef());
+		FToolMenuContext ToolMenuContext(HostedAssetEditorToolkit->GetToolkitCommands(), FExtender::Combine(MenuExtenders).ToSharedRef());
 		IMainFrameModule& MainFrameModule = FModuleManager::LoadModuleChecked<IMainFrameModule>( "MainFrame" );
-		DefaultMenuWidget = MainFrameModule.MakeMainMenu( MyTabManager, AssetEditorMenuName, EditorMenuContext );
+		DefaultMenuWidget = MainFrameModule.MakeMainMenu( MyTabManager, AssetEditorMenuName, ToolMenuContext );
 
 		MenuWidgetContent->SetContent(DefaultMenuWidget.ToSharedRef());
 	}
diff --git a/Engine/Source/Editor/UnrealEd/Public/Kismet2/ComponentEditorUtils.h b/Engine/Source/Editor/UnrealEd/Public/Kismet2/ComponentEditorUtils.h
--- a/Engine/Source/Editor/UnrealEd/Public/Kismet2/ComponentEditorUtils.h
+++ b/Engine/Source/Editor/UnrealEd/Public/Kismet2/ComponentEditorUtils.h
@@ -9,7 +9,7 @@
 #include "GameFramework/Actor.h"
 #include "UObject/UObjectHash.h"
 
-class UEditorMenu;
+class UToolMenu;
 class UMaterialInterface;
 
 class UNREALED_API FComponentEditorUtils
@@ -246,7 +246,7 @@ public:
 	* @param Menu Used to register the menu options
 	* @param SelectedComponents The selected components to create menu options for
 	*/
-	static void FillComponentContextMenuOptions(UEditorMenu* Menu, const TArray<UActorComponent*>& SelectedComponents);
+	static void FillComponentContextMenuOptions(UToolMenu* Menu, const TArray<UActorComponent*>& SelectedComponents);
 
 	/**
 	 * Tries to find a match for ComponentInstance in the ComponentList. First by name and then if multiple Components have a matching name try to match the SceneComponent hierarchy to find the best match.
diff --git a/Engine/Source/Editor/UnrealEd/Public/Kismet2/KismetEditorUtilities.h b/Engine/Source/Editor/UnrealEd/Public/Kismet2/KismetEditorUtilities.h
--- a/Engine/Source/Editor/UnrealEd/Public/Kismet2/KismetEditorUtilities.h
+++ b/Engine/Source/Editor/UnrealEd/Public/Kismet2/KismetEditorUtilities.h
@@ -10,7 +10,7 @@
 #include "Engine/Blueprint.h"
 #include "Engine/BlueprintGeneratedClass.h"
 
-class UEditorMenu;
+class UToolMenu;
 class IBlueprintEditor;
 class UEdGraph;
 struct Rect;
@@ -277,7 +277,7 @@ public:
 	static bool AnyBoundLevelScriptEventForActor(AActor* Actor, bool bCouldAddAny);
 
 	/** It lists bounded LevelScriptEvents for given actor */
-	static void AddLevelScriptEventOptionsForActor(UEditorMenu* Menu, TWeakObjectPtr<AActor> ActorPtr, bool bExistingEvents, bool bNewEvents, bool bOnlyEventName);
+	static void AddLevelScriptEventOptionsForActor(UToolMenu* Menu, TWeakObjectPtr<AActor> ActorPtr, bool bExistingEvents, bool bNewEvents, bool bOnlyEventName);
 	
 	/** Return information about the given macro graph */
 	static void GetInformationOnMacro(UEdGraph* MacroGraph, /*out*/ class UK2Node_Tunnel*& EntryNode, /*out*/ class UK2Node_Tunnel*& ExitNode, bool& bIsPure);
diff --git a/Engine/Source/Editor/UnrealEd/Public/Toolkits/AssetEditorToolkit.h b/Engine/Source/Editor/UnrealEd/Public/Toolkits/AssetEditorToolkit.h
--- a/Engine/Source/Editor/UnrealEd/Public/Toolkits/AssetEditorToolkit.h
+++ b/Engine/Source/Editor/UnrealEd/Public/Toolkits/AssetEditorToolkit.h
@@ -185,9 +185,9 @@ public:
 	/** Regenerates the menubar and toolbar widgets */
 	void RegenerateMenusAndToolbars();
 
-	/** Get name used by editor menu subsystem */
-	FName GetEditorMenuToolbarName() const;
-	FName GetEditorMenuAppName() const;
+	/** Get name used by tool menu */
+	FName GetToolMenuToolbarName() const;
+	FName GetToolMenuAppName() const;
 
 	/** Called at the end of RegenerateMenusAndToolbars() */
 	virtual void PostRegenerateMenusAndToolbars() { }
diff --git a/Engine/Source/Editor/WorldBrowser/Private/SWorldHierarchyImpl.cpp b/Engine/Source/Editor/WorldBrowser/Private/SWorldHierarchyImpl.cpp
--- a/Engine/Source/Editor/WorldBrowser/Private/SWorldHierarchyImpl.cpp
+++ b/Engine/Source/Editor/WorldBrowser/Private/SWorldHierarchyImpl.cpp
@@ -10,7 +10,7 @@
 
 #include "Framework/Application/SlateApplication.h"
 #include "Framework/MultiBox/MultiBoxBuilder.h"
-#include "EditorMenuSubsystem.h"
+#include "ToolMenus.h"
 #include "Widgets/Layout/SSeparator.h"
 #include "Widgets/Images/SImage.h"
 #include "Widgets/Input/SComboButton.h"
@@ -459,15 +459,15 @@ TSharedPtr<SWidget> SWorldHierarchyImpl::ConstructLevelContextMenu() const
 
 	if (!WorldModel->IsReadOnly())
 	{
-		UEditorMenuSubsystem* EditorMenus = UEditorMenuSubsystem::Get();
+		UToolMenus* ToolMenus = UToolMenus::Get();
 		static const FName MenuName = "WorldBrowser.WorldHierarchy.LevelContextMenu";
-		if (!EditorMenus->IsMenuRegistered(MenuName))
+		if (!ToolMenus->IsMenuRegistered(MenuName))
 		{
-			EditorMenus->RegisterMenu(MenuName);
+			ToolMenus->RegisterMenu(MenuName);
 		}
 
-		FEditorMenuContext Context(WorldModel->GetCommandList(), TSharedPtr<FExtender>());
-		UEditorMenu* Menu = EditorMenus->GenerateMenu(MenuName, Context);
+		FToolMenuContext Context(WorldModel->GetCommandList(), TSharedPtr<FExtender>());
+		UToolMenu* Menu = ToolMenus->GenerateMenu(MenuName, Context);
 
 		TArray<WorldHierarchy::FWorldTreeItemPtr> SelectedItems = GetSelectedTreeItems();
 
@@ -482,7 +482,7 @@ TSharedPtr<SWidget> SWorldHierarchyImpl::ConstructLevelContextMenu() const
 			RootTreeItems[0]->GenerateContextMenu(Menu, *this);
 		}
 
-		Menu->AddDynamicSection("HierarchyDynamicSection", FNewEditorMenuDelegateLegacy::CreateLambda([=](FMenuBuilder& MenuBuilder, UEditorMenu* Menu)
+		Menu->AddDynamicSection("HierarchyDynamicSection", FNewToolMenuDelegateLegacy::CreateLambda([=](FMenuBuilder& MenuBuilder, UToolMenu* Menu)
 		{
 			WorldModel->BuildHierarchyMenu(MenuBuilder);
 		}));
@@ -507,8 +507,8 @@ TSharedPtr<SWidget> SWorldHierarchyImpl::ConstructLevelContextMenu() const
 
 			if (bAllSelectedItemsCanMove && FLevelFolders::IsAvailable())
 			{
-				FEditorMenuSection& Section = Menu->AddSection("Section");
-				Section.AddEntry(FEditorMenuEntry::InitSubMenu(
+				FToolMenuSection& Section = Menu->AddSection("Section");
+				Section.AddEntry(FToolMenuEntry::InitSubMenu(
 					Menu->MenuName,
 					"MoveSelectionTo",
 					LOCTEXT("MoveSelectionTo", "Move To"),
@@ -519,8 +519,8 @@ TSharedPtr<SWidget> SWorldHierarchyImpl::ConstructLevelContextMenu() const
 
 			if (bOnlyFoldersSelected)
 			{
-				FEditorMenuSection& Section = Menu->AddSection("Section");
-				Section.AddEntry(FEditorMenuEntry::InitSubMenu(
+				FToolMenuSection& Section = Menu->AddSection("Section");
+				Section.AddEntry(FToolMenuEntry::InitSubMenu(
 					Menu->MenuName,
 					"SelectSubmenu",
 					LOCTEXT("SelectSubmenu", "Select"),
@@ -530,7 +530,7 @@ TSharedPtr<SWidget> SWorldHierarchyImpl::ConstructLevelContextMenu() const
 			}
 		}
 
-		MenuWidget = EditorMenus->GenerateWidget(Menu);
+		MenuWidget = ToolMenus->GenerateWidget(Menu);
 	}
 
 	return MenuWidget;
diff --git a/Engine/Source/Editor/WorldBrowser/Private/WorldTreeItemTypes.cpp b/Engine/Source/Editor/WorldBrowser/Private/WorldTreeItemTypes.cpp
--- a/Engine/Source/Editor/WorldBrowser/Private/WorldTreeItemTypes.cpp
+++ b/Engine/Source/Editor/WorldBrowser/Private/WorldTreeItemTypes.cpp
@@ -7,7 +7,7 @@
 #include "Misc/PackageName.h"
 #include "Engine/Engine.h"
 #include "SWorldHierarchyImpl.h"
-#include "EditorMenuSubsystem.h"
+#include "ToolMenus.h"
 #include "EditorStyleSet.h"
 #include "AssetSelection.h"
 #include "SWorldHierarchyImpl.h"
@@ -476,7 +476,7 @@ namespace WorldHierarchy
 		return Parent.IsValid();
 	}
 
-	void FLevelModelTreeItem::GenerateContextMenu(UEditorMenu* Menu, const SWorldHierarchyImpl& Hierarchy)
+	void FLevelModelTreeItem::GenerateContextMenu(UToolMenu* Menu, const SWorldHierarchyImpl& Hierarchy)
 	{
 		if (!Parent.IsValid() && WorldModel.Pin()->HasFolderSupport())
 		{
@@ -486,7 +486,7 @@ namespace WorldHierarchy
 			FName RootPath = LevelModel.Pin()->GetFolderPath();
 			auto NewFolderAction = FExecuteAction::CreateSP(const_cast<SWorldHierarchyImpl*>(&Hierarchy), &SWorldHierarchyImpl::CreateFolder, LevelModel.Pin(), RootPath);
 
-			FEditorMenuSection& Section = Menu->AddSection("Section");
+			FToolMenuSection& Section = Menu->AddSection("Section");
 			Section.AddMenuEntry("CreateFolder", LOCTEXT("CreateFolder", "Create Folder"), FText(), NewFolderIcon, FUIAction(NewFolderAction));
 		}
 	}
@@ -904,7 +904,7 @@ namespace WorldHierarchy
 		return FEditorStyle::Get().GetBrush("WorldBrowser.FolderClosed");
 	}
 
-	void FFolderTreeItem::GenerateContextMenu(UEditorMenu* Menu, const SWorldHierarchyImpl& Hierarchy)
+	void FFolderTreeItem::GenerateContextMenu(UToolMenu* Menu, const SWorldHierarchyImpl& Hierarchy)
 	{
 		// Folder items should be able to create subfolders, rename themselves, or delete themselves from the tree
 		const FSlateIcon NewFolderIcon(FEditorStyle::GetStyleSetName(), "WorldBrowser.NewFolderIcon");
@@ -917,7 +917,7 @@ namespace WorldHierarchy
 		auto RenameFolderAction = FExecuteAction::CreateSP(const_cast<SWorldHierarchyImpl*>(&Hierarchy), &SWorldHierarchyImpl::InitiateRename, AsShared());
 		auto DeleteFolderAction = FExecuteAction::CreateSP(const_cast<SWorldHierarchyImpl*>(&Hierarchy), &SWorldHierarchyImpl::DeleteFolders, Folders, /*bTransactional*/ true);
 
-		FEditorMenuSection& Section = Menu->AddSection("Section");
+		FToolMenuSection& Section = Menu->AddSection("Section");
 		Section.AddMenuEntry("CreateSubFolder", LOCTEXT("CreateSubFolder", "Create Subfolder"), FText(), NewFolderIcon, FUIAction(NewFolderAction));
 		Section.AddMenuEntry("RenameFolder", LOCTEXT("RenameFolder", "Rename"), FText(), FSlateIcon(FEditorStyle::GetStyleSetName(), "ContentBrowser.AssetActions.Rename"), FUIAction(RenameFolderAction));
 		Section.AddMenuEntry("DeleteFolder", LOCTEXT("DeleteFolder", "Delete"), FText(), FSlateIcon(FEditorStyle::GetStyleSetName(), "ContentBrowser.AssetActions.Delete"), FUIAction(DeleteFolderAction));
diff --git a/Engine/Source/Editor/WorldBrowser/Private/WorldTreeItemTypes.h b/Engine/Source/Editor/WorldBrowser/Private/WorldTreeItemTypes.h
--- a/Engine/Source/Editor/WorldBrowser/Private/WorldTreeItemTypes.h
+++ b/Engine/Source/Editor/WorldBrowser/Private/WorldTreeItemTypes.h
@@ -87,7 +87,7 @@ namespace WorldHierarchy
 		virtual const FSlateBrush* GetHierarchyItemBrush() const override;
 
 		virtual bool CanChangeParents() const override;
-		virtual void GenerateContextMenu(UEditorMenu* Menu, const SWorldHierarchyImpl& Hierarchy) override;
+		virtual void GenerateContextMenu(UToolMenu* Menu, const SWorldHierarchyImpl& Hierarchy) override;
 
 	public:
 		virtual FLevelModelTreeItem* GetAsLevelModelTreeItem() const override { return const_cast<FLevelModelTreeItem*>(this); }
@@ -165,7 +165,7 @@ namespace WorldHierarchy
 		virtual float GetHierarchyItemBrushWidth() const { return 16.0f; }
 
 		virtual bool CanChangeParents() const { return true; }
-		virtual void GenerateContextMenu(UEditorMenu* Menu, const SWorldHierarchyImpl& Hierarchy) override;
+		virtual void GenerateContextMenu(UToolMenu* Menu, const SWorldHierarchyImpl& Hierarchy) override;
 
 		/** Sets the new path of the folder. This includes the leaf name. This will not rename any children in this folder */
 		void SetNewPath(FName NewPath);
diff --git a/Engine/Source/Editor/WorldBrowser/Public/IWorldTreeItem.h b/Engine/Source/Editor/WorldBrowser/Public/IWorldTreeItem.h
--- a/Engine/Source/Editor/WorldBrowser/Public/IWorldTreeItem.h
+++ b/Engine/Source/Editor/WorldBrowser/Public/IWorldTreeItem.h
@@ -14,7 +14,7 @@ class SWidget;
 class SWorldHierarchy;
 class FLevelCollectionModel;
 class SWorldHierarchyImpl;
-class UEditorMenu;
+class UToolMenu;
 
 namespace WorldHierarchy
 {
@@ -270,7 +270,7 @@ namespace WorldHierarchy
 		virtual bool CanChangeParents() const = 0;
 
 		/** Generates a context menu option for this item if and only if it's the only item selected in the hierarchy */
-		virtual void GenerateContextMenu(UEditorMenu* Menu, const SWorldHierarchyImpl& Hierarchy) {}
+		virtual void GenerateContextMenu(UToolMenu* Menu, const SWorldHierarchyImpl& Hierarchy) {}
 
 		/** Sets the item's visible status */
 		virtual void SetVisible(bool bVisible) = 0;
