Enhacement JIRA UETOOL-1721: Introduced the capability to load/save/remove Editor UI Layouts, so that the user can customize the Editor and save/share it accross machines or people.
Additional enhancements: Re-loading a layout no longer requires the whole editor to restart. Instead, only the tabs are closed and re-populated (only creating a small "flash" in the editor rather than restarting it).
Previous Editor bugs fixed:
 - Bug 1: AssetEditorToolkit.cpp::FAssetEditorToolkit::RestoreFromLayout was wrongly saving the layout in GEditorIni rather than GEditorLayoutIni.
 - Bug 2: When the layout is saved, the instance ID of the custom toolkits was being saved as well. Given that this number keep increasing after each custom tab is created (no matter if the previous ones were already closed), the editor will keep thinking that the loaded layout is no longer the currrent layout. Fix in TabManager.h::FTabId::ToString().
 - Bug 3: FLayoutSaveRestore::SaveToConfig was taking as input argument FString rather than const FString&.

diff --git a/Engine/Config/BaseEditorLayout.ini b/Engine/Config/BaseEditorLayout.ini
deleted file mode 100644
diff --git a/Engine/Config/BaseEditorPerProjectUserSettings.ini b/Engine/Config/BaseEditorPerProjectUserSettings.ini
--- a/Engine/Config/BaseEditorPerProjectUserSettings.ini
+++ b/Engine/Config/BaseEditorPerProjectUserSettings.ini
@@ -72,6 +72,9 @@ VertexColor=true
 MapSkeletalMotionToRoot=False
 
 [/Script/EditorStyle.EditorStyleSettings]
+; Whether to enable the Editor UI Layout configuration tools for the user
+bEnableUserEditorLayoutManagement=True
+
 ; Applies a color vision deficiency filter to the entire editor
 ColorVisionDeficiencyPreviewType=NormalVision
 ColorVisionDeficiencySeverity=3
diff --git a/Engine/Source/Editor/EditorStyle/Private/EditorStyleClasses.cpp b/Engine/Source/Editor/EditorStyle/Private/EditorStyleClasses.cpp
--- a/Engine/Source/Editor/EditorStyle/Private/EditorStyleClasses.cpp
+++ b/Engine/Source/Editor/EditorStyle/Private/EditorStyleClasses.cpp
@@ -16,6 +16,8 @@
 UEditorStyleSettings::UEditorStyleSettings( const FObjectInitializer& ObjectInitializer )
 	: Super(ObjectInitializer)
 {
+	bEnableUserEditorLayoutManagement = true;
+
 	SelectionColor = FLinearColor(0.828f, 0.364f, 0.003f);
 	InactiveSelectionColor = FLinearColor(0.25f, 0.25f, 0.25f);
 	PressedSelectionColor = FLinearColor(0.701f, 0.225f, 0.003f);
diff --git a/Engine/Source/Editor/EditorStyle/Public/Classes/EditorStyleSettings.h b/Engine/Source/Editor/EditorStyle/Public/Classes/EditorStyleSettings.h
--- a/Engine/Source/Editor/EditorStyle/Public/Classes/EditorStyleSettings.h
+++ b/Engine/Source/Editor/EditorStyle/Public/Classes/EditorStyleSettings.h
@@ -48,6 +48,13 @@ public:
 	UPROPERTY(EditAnywhere, Category=UserInterface, meta = (ConfigRestartRequired = true, DisplayName="Enable High DPI Support"))
 	bool bEnableHighDPIAwareness;
 
+	/**
+	 * Whether to enable the Editor UI Layout configuration tools for the user.
+	 * If disabled, the "Save Layout As" and "Remove Layout" menus will be removed, as well as the "Import Layout..." sub-menu.
+	 */
+	UPROPERTY(EditAnywhere, config, Category = UserInterface)
+	bool bEnableUserEditorLayoutManagement;
+
 	/** Applies a color vision deficiency filter to the entire editor */
 	UPROPERTY(EditAnywhere, config, Category = "Accessibility")
 	EColorVisionDeficiency ColorVisionDeficiencyPreviewType;
diff --git a/Engine/Source/Editor/MainFrame/Private/Frame/MainFrameActions.cpp b/Engine/Source/Editor/MainFrame/Private/Frame/MainFrameActions.cpp
--- a/Engine/Source/Editor/MainFrame/Private/Frame/MainFrameActions.cpp
+++ b/Engine/Source/Editor/MainFrame/Private/Frame/MainFrameActions.cpp
@@ -29,12 +29,10 @@
 #include "Interfaces/ITargetPlatformManagerModule.h"
 #include "PlatformInfo.h"
 #include "EditorStyleSet.h"
-#include "Editor/EditorPerProjectUserSettings.h"
 #include "Settings/EditorExperimentalSettings.h"
 #include "CookerSettings.h"
 #include "UnrealEdMisc.h"
 #include "FileHelpers.h"
-#include "Dialogs/Dialogs.h"
 #include "EditorAnalytics.h"
 #include "LevelEditor.h"
 #include "Interfaces/IProjectTargetPlatformEditorModule.h"
@@ -42,11 +40,11 @@
 #include "Misc/ConfigCacheIni.h"
 #include "MainFrameModule.h"
 #include "Widgets/Docking/SDockTab.h"
-#include "Framework/Notifications/NotificationManager.h"
 #include "Widgets/Notifications/SNotificationList.h"
 #include "Framework/Commands/GenericCommands.h"
 #include "Dialogs/SOutputLogDialog.h"
 #include "IUATHelperModule.h"
+#include "Menus/LayoutsMenu.h"
 
 #include "Settings/EditorSettings.h"
 #include "AnalyticsEventAttribute.h"
@@ -195,11 +193,8 @@ void FMainFrameCommands::RegisterCommands()
 	UI_COMMAND( CreditsUnrealEd, "Credits", "Displays application credits", EUserInterfaceActionType::Button, FInputChord() );
 	ActionList->MapAction( CreditsUnrealEd, FExecuteAction::CreateStatic(&FMainFrameActionCallbacks::CreditsUnrealEd_Execute) );
 
-	UI_COMMAND( ResetLayout, "Reset Layout...", "Make a backup of your user settings and reset the layout customizations", EUserInterfaceActionType::Button, FInputChord() );
-	ActionList->MapAction( ResetLayout, FExecuteAction::CreateStatic( &FMainFrameActionCallbacks::ResetLayout) );
-
-	UI_COMMAND( SaveLayout, "Save Layout", "Save the layout customizations", EUserInterfaceActionType::Button, FInputChord() );
-	ActionList->MapAction( SaveLayout, FExecuteAction::CreateStatic( &FMainFrameActionCallbacks::SaveLayout) );
+	// Layout commands
+	RegisterLayoutCommands();
 
 #if !PLATFORM_MAC && !PLATFORM_LINUX // Fullscreen mode in the editor is currently unsupported on Mac and Linux
 	UI_COMMAND( ToggleFullscreen, "Enable Fullscreen", "Enables fullscreen mode for the application, expanding across the entire monitor", EUserInterfaceActionType::ToggleButton, FInputChord(EModifierKey::Shift, EKeys::F11) );
@@ -216,6 +211,96 @@ void FMainFrameCommands::RegisterCommands()
 	FGlobalEditorCommonCommands::MapActions(ActionList);
 }
 
+TSharedRef<FUICommandInfo> CreateFUICommandInfoInternal(const TSharedRef<FBindingContext>& InThisAsShared, const FName& InCommandName)
+{
+	// Default FUICommandInfoDecl command
+	return FUICommandInfoDecl(
+			InThisAsShared,
+			InCommandName,
+			FText::FromString("InLabel"),
+			FText::FromString("InDesc"))
+		.UserInterfaceType(EUserInterfaceActionType::Check)
+		.DefaultChord(FInputChord());
+}
+
+void FMainFrameCommands::RegisterLayoutCommands()
+{
+	// Load layout commands
+	UI_COMMAND(MainFrameLayoutCommands.ImportLayout, "Import Layout...", "Import a custom layout (or set of layouts) from a different directory and load it into your current instance of the Unreal Editor UI", EUserInterfaceActionType::Button, FInputChord());
+	ActionList->MapAction(MainFrameLayoutCommands.ImportLayout, FExecuteAction::CreateStatic(&FLayoutsMenuLoad::ImportLayout));
+
+	// Save layout commands
+	UI_COMMAND(MainFrameLayoutCommands.SaveLayoutAs, "Save Layout As...", "Save the current layout customization on disk so it can be loaded later", EUserInterfaceActionType::Button, FInputChord());
+	ActionList->MapAction(MainFrameLayoutCommands.SaveLayoutAs, FExecuteAction::CreateStatic(&FLayoutsMenuSave::SaveLayoutAs));
+	UI_COMMAND(MainFrameLayoutCommands.ExportLayout, "Export Layout...", "Export the custom layout customization to a different directory", EUserInterfaceActionType::Button, FInputChord());
+	ActionList->MapAction(MainFrameLayoutCommands.ExportLayout, FExecuteAction::CreateStatic(&FLayoutsMenuSave::ExportLayout));
+
+	// Remove layout commands
+	UI_COMMAND(MainFrameLayoutCommands.RemoveUserLayouts, "Remove All User Layouts...", "Remove all the layout customizations created by the user", EUserInterfaceActionType::Button, FInputChord());
+	ActionList->MapAction(
+		MainFrameLayoutCommands.RemoveUserLayouts,
+		FExecuteAction::CreateStatic(&FLayoutsMenuRemove::RemoveUserLayouts),
+		FCanExecuteAction::CreateStatic(&FLayoutsMenuBase::IsThereUserLayouts)
+	);
+
+	// Load/Override/Remove layout commands
+	const int32 MaxLayouts = 50;
+	for (int32 CurrentLayoutIndex = 0; CurrentLayoutIndex < MaxLayouts; ++CurrentLayoutIndex)
+	{
+		// NOTE: The actual label and tool-tip will be overridden at runtime when the command is bound to a menu item, however
+		// we still need to set one here so that the key bindings UI can function properly
+		// Thus, we remove any (NS)LOCTEXT from here to avoid unnecessary translations.
+
+		// Load Layout
+		{
+			// Default Layouts
+			TSharedRef<FUICommandInfo> LoadLayout = CreateFUICommandInfoInternal(this->AsShared(), FName(*FString::Printf(TEXT("LoadLayout%i"), CurrentLayoutIndex)));
+			MainFrameLayoutCommands.LoadLayoutCommands.Add(LoadLayout);
+			ActionList->MapAction(MainFrameLayoutCommands.LoadLayoutCommands[CurrentLayoutIndex], FExecuteAction::CreateStatic(&FLayoutsMenuLoad::LoadLayout, CurrentLayoutIndex),
+				FCanExecuteAction::CreateStatic(&FLayoutsMenuLoad::CanLoadChooseLayout, CurrentLayoutIndex),
+				FIsActionChecked::CreateStatic(&FLayoutsMenuBase::IsLayoutChecked, CurrentLayoutIndex));
+			// User Layouts
+			TSharedRef<FUICommandInfo> LoadUserLayout = CreateFUICommandInfoInternal(this->AsShared(), FName(*FString::Printf(TEXT("LoadUserLayout%i"), CurrentLayoutIndex)));
+			MainFrameLayoutCommands.LoadUserLayoutCommands.Add(LoadUserLayout);
+			ActionList->MapAction(MainFrameLayoutCommands.LoadUserLayoutCommands[CurrentLayoutIndex], FExecuteAction::CreateStatic(&FLayoutsMenuLoad::LoadUserLayout, CurrentLayoutIndex),
+				FCanExecuteAction::CreateStatic(&FLayoutsMenuLoad::CanLoadChooseUserLayout, CurrentLayoutIndex),
+				FIsActionChecked::CreateStatic(&FLayoutsMenuBase::IsUserLayoutChecked, CurrentLayoutIndex));
+		}
+
+		// Override Layout
+		{
+			// Default Layouts
+			TSharedRef<FUICommandInfo> OverrideLayout = CreateFUICommandInfoInternal(this->AsShared(), FName(*FString::Printf(TEXT("OverrideLayout%i"), CurrentLayoutIndex)));
+			MainFrameLayoutCommands.OverrideLayoutCommands.Add(OverrideLayout);
+			ActionList->MapAction(MainFrameLayoutCommands.OverrideLayoutCommands[CurrentLayoutIndex], FExecuteAction::CreateStatic(&FLayoutsMenuSave::OverrideLayout, CurrentLayoutIndex),
+				FCanExecuteAction::CreateStatic(&FLayoutsMenuSave::CanSaveChooseLayout, CurrentLayoutIndex),
+				FIsActionChecked::CreateStatic(&FLayoutsMenuBase::IsLayoutChecked, CurrentLayoutIndex));
+			// User Layouts
+			TSharedRef<FUICommandInfo> OverrideUserLayout = CreateFUICommandInfoInternal(this->AsShared(), FName(*FString::Printf(TEXT("OverrideUserLayout%i"), CurrentLayoutIndex)));
+			MainFrameLayoutCommands.OverrideUserLayoutCommands.Add(OverrideUserLayout);
+			ActionList->MapAction(MainFrameLayoutCommands.OverrideUserLayoutCommands[CurrentLayoutIndex], FExecuteAction::CreateStatic(&FLayoutsMenuSave::OverrideUserLayout, CurrentLayoutIndex),
+				FCanExecuteAction::CreateStatic(&FLayoutsMenuSave::CanSaveChooseUserLayout, CurrentLayoutIndex),
+				FIsActionChecked::CreateStatic(&FLayoutsMenuBase::IsUserLayoutChecked, CurrentLayoutIndex));
+		}
+
+		// Remove Layout
+		{
+			// Default Layouts
+			TSharedRef<FUICommandInfo> RemoveLayout = CreateFUICommandInfoInternal(this->AsShared(), FName(*FString::Printf(TEXT("RemoveLayout%i"), CurrentLayoutIndex)));
+			MainFrameLayoutCommands.RemoveLayoutCommands.Add(RemoveLayout);
+			ActionList->MapAction(MainFrameLayoutCommands.RemoveLayoutCommands[CurrentLayoutIndex], FExecuteAction::CreateStatic(&FLayoutsMenuRemove::RemoveLayout, CurrentLayoutIndex),
+				FCanExecuteAction::CreateStatic(&FLayoutsMenuRemove::CanRemoveChooseLayout, CurrentLayoutIndex),
+				FIsActionChecked::CreateStatic(&FLayoutsMenuBase::IsLayoutChecked, CurrentLayoutIndex));
+			// User Layouts
+			TSharedRef<FUICommandInfo> RemoveUserLayout = CreateFUICommandInfoInternal(this->AsShared(), FName(*FString::Printf(TEXT("RemoveUserLayout%i"), CurrentLayoutIndex)));
+			MainFrameLayoutCommands.RemoveUserLayoutCommands.Add(RemoveUserLayout);
+			ActionList->MapAction(MainFrameLayoutCommands.RemoveUserLayoutCommands[CurrentLayoutIndex], FExecuteAction::CreateStatic(&FLayoutsMenuRemove::RemoveUserLayout, CurrentLayoutIndex),
+				FCanExecuteAction::CreateStatic(&FLayoutsMenuRemove::CanRemoveChooseUserLayout, CurrentLayoutIndex),
+				FIsActionChecked::CreateStatic(&FLayoutsMenuBase::IsUserLayoutChecked, CurrentLayoutIndex));
+		}
+	}
+}
+
 FReply FMainFrameActionCallbacks::OnUnhandledKeyDownEvent(const FKeyEvent& InKeyEvent)
 {
 	if(!GIsSlowTask)
@@ -877,69 +962,6 @@ void FMainFrameActionCallbacks::OpenBackupDirectory( FString BackupFile )
 	FPlatformProcess::LaunchFileInDefaultExternalApplication(*FPaths::GetPath(FPaths::ConvertRelativePathToFull(BackupFile)));
 }
 
-void FMainFrameActionCallbacks::ResetLayout()
-{
-	if(EAppReturnType::Ok != OpenMsgDlgInt(EAppMsgType::OkCancel, LOCTEXT( "ActionRestartMsg", "This action requires the editor to restart; you will be prompted to save any changes. Continue?" ), LOCTEXT( "ResetUILayout_Title", "Reset UI Layout" ) ) )
-	{
-		return;
-	}
-
-	// make a backup
-	GetMutableDefault<UEditorPerProjectUserSettings>()->SaveConfig();
-
-	FString BackupEditorLayoutIni = FString::Printf(TEXT("%s_Backup.ini"), *FPaths::GetBaseFilename(GEditorLayoutIni, false));
-
-	if( COPY_Fail == IFileManager::Get().Copy(*BackupEditorLayoutIni, *GEditorLayoutIni) )
-	{
-		FMessageLog EditorErrors("EditorErrors");
-		if(!FPaths::FileExists(GEditorLayoutIni))
-		{
-			FFormatNamedArguments Arguments;
-			Arguments.Add(TEXT("FileName"), FText::FromString(GEditorLayoutIni));
-			EditorErrors.Warning(FText::Format(LOCTEXT("UnsuccessfulBackup_NoExist_Notification", "Unsuccessful backup! {FileName} does not exist!"), Arguments));
-		}
-		else if(IFileManager::Get().IsReadOnly(*BackupEditorLayoutIni))
-		{
-			FFormatNamedArguments Arguments;
-			Arguments.Add(TEXT("FileName"), FText::FromString(FPaths::ConvertRelativePathToFull(BackupEditorLayoutIni)));
-			EditorErrors.Warning(FText::Format(LOCTEXT("UnsuccessfulBackup_ReadOnly_Notification", "Unsuccessful backup! {FileName} is read-only!"), Arguments));
-		}
-		else
-		{
-			// We don't specifically know why it failed, this is a fallback.
-			FFormatNamedArguments Arguments;
-			Arguments.Add(TEXT("SourceFileName"), FText::FromString(GEditorLayoutIni));
-			Arguments.Add(TEXT("BackupFileName"), FText::FromString(FPaths::ConvertRelativePathToFull(BackupEditorLayoutIni)));
-			EditorErrors.Warning(FText::Format(LOCTEXT("UnsuccessfulBackup_Fallback_Notification", "Unsuccessful backup of {SourceFileName} to {BackupFileName}"), Arguments));
-		}
-		EditorErrors.Notify(LOCTEXT("BackupUnsuccessful_Title", "Backup Unsuccessful!"));
-	}
-	else
-	{
-		FNotificationInfo ErrorNotification( FText::GetEmpty() );
-		ErrorNotification.bFireAndForget = true;
-		ErrorNotification.ExpireDuration = 3.0f;
-		ErrorNotification.bUseThrobber = true;
-		ErrorNotification.Hyperlink = FSimpleDelegate::CreateStatic(&FMainFrameActionCallbacks::OpenBackupDirectory, BackupEditorLayoutIni);
-		ErrorNotification.HyperlinkText = LOCTEXT("SuccessfulBackup_Notification_Hyperlink", "Open Directory");
-		ErrorNotification.Text = LOCTEXT("SuccessfulBackup_Notification", "Backup Successful!");
-		ErrorNotification.Image = FEditorStyle::GetBrush(TEXT("NotificationList.SuccessImage"));
-		FSlateNotificationManager::Get().AddNotification(ErrorNotification);
-	}
-
-	// reset layout & restart Editor
-	FUnrealEdMisc::Get().AllowSavingLayoutOnClose(false);
-	FUnrealEdMisc::Get().RestartEditor(false);
-}
-
-void FMainFrameActionCallbacks::SaveLayout()
-{
-	FGlobalTabmanager::Get()->SaveAllVisualState();
-
-	// Write the saved state's config to disk
-	GConfig->Flush( false, GEditorLayoutIni );
-}
-
 void FMainFrameActionCallbacks::ToggleFullscreen_Execute()
 {
 #if !PLATFORM_MAC && !PLATFORM_LINUX // Fullscreen mode in the editor is currently unsupported on Mac or Linux
diff --git a/Engine/Source/Editor/MainFrame/Private/Frame/MainFrameActions.h b/Engine/Source/Editor/MainFrame/Private/Frame/MainFrameActions.h
--- a/Engine/Source/Editor/MainFrame/Private/Frame/MainFrameActions.h
+++ b/Engine/Source/Editor/MainFrame/Private/Frame/MainFrameActions.h
@@ -11,6 +11,27 @@
 
 class FUICommandList;
 
+/**
+ * Unreal editor main frame actions related to the Layout section
+ */
+class FMainFrameLayoutCommands
+{
+public:
+	// Load-based commands
+	TArray<TSharedPtr<FUICommandInfo>> LoadLayoutCommands;
+	TArray<TSharedPtr<FUICommandInfo>> LoadUserLayoutCommands;
+	TSharedPtr< FUICommandInfo > ImportLayout;
+	// Save-based commands
+	TArray<TSharedPtr<FUICommandInfo>> OverrideLayoutCommands;
+	TArray<TSharedPtr<FUICommandInfo>> OverrideUserLayoutCommands;
+	TSharedPtr< FUICommandInfo > SaveLayoutAs;
+	TSharedPtr< FUICommandInfo > ExportLayout;
+	// Remove-based commands
+	TArray<TSharedPtr<FUICommandInfo>> RemoveLayoutCommands;
+	TArray<TSharedPtr<FUICommandInfo>> RemoveUserLayoutCommands;
+	TSharedPtr< FUICommandInfo > RemoveUserLayouts;
+};
+
 /**
  * Unreal editor main frame actions
  */
@@ -61,8 +82,9 @@ public:
 	TSharedPtr< FUICommandInfo > AboutUnrealEd;
 	TSharedPtr< FUICommandInfo > CreditsUnrealEd;
 
-	TSharedPtr< FUICommandInfo > ResetLayout;
-	TSharedPtr< FUICommandInfo > SaveLayout;
+	// Layout
+	FMainFrameLayoutCommands MainFrameLayoutCommands;
+
 	TSharedPtr< FUICommandInfo > ToggleFullscreen;
 
 	virtual void RegisterCommands() override;
@@ -73,6 +95,8 @@ private:
 	/** Console command for toggling full screen.  We need this to expose the full screen toggle action to
 	    the game UI system for play-in-editor view ports */
 	FAutoConsoleCommand ToggleFullscreenConsoleCommand;
+
+	void RegisterLayoutCommands();
 };
 
 
@@ -215,12 +239,6 @@ public:
 	/** Opens the directory where the backup for preferences is stored. */
 	static void OpenBackupDirectory( FString BackupFile );
 
-	/** Resets the visual state of the editor */
-	static void ResetLayout();
-
-	/** Save the visual state of the editor */
-	static void SaveLayout();
-
 	/** Toggle the level editor's fullscreen mode */
 	static void ToggleFullscreen_Execute();
 
diff --git a/Engine/Source/Editor/MainFrame/Private/Frame/MainFrameHandler.cpp b/Engine/Source/Editor/MainFrame/Private/Frame/MainFrameHandler.cpp
--- a/Engine/Source/Editor/MainFrame/Private/Frame/MainFrameHandler.cpp
+++ b/Engine/Source/Editor/MainFrame/Private/Frame/MainFrameHandler.cpp
@@ -54,10 +54,6 @@ void FMainFrameHandler::ShutDownEditor()
 	{
 		GlobalTabManager->SaveAllVisualState();
 	}
-	else
-	{
-		GConfig->EmptySection(TEXT("EditorLayouts"), *GEditorLayoutIni);
-	}
 
 	// Clear the callback for destructionfrom the main tab; otherwise it will re-enter this shutdown function.
 	if (MainTabPtr.IsValid())
diff --git a/Engine/Source/Editor/MainFrame/Private/MainFrameModule.cpp b/Engine/Source/Editor/MainFrame/Private/MainFrameModule.cpp
--- a/Engine/Source/Editor/MainFrame/Private/MainFrameModule.cpp
+++ b/Engine/Source/Editor/MainFrame/Private/MainFrameModule.cpp
@@ -233,6 +233,20 @@ void FMainFrameModule::CreateDefaultMainFrame( const bool bStartImmersive, const
 	}
 }
 
+void FMainFrameModule::RecreateDefaultMainFrame(const bool bStartImmersive, const bool bStartPIE)
+{
+	// Clean previous default main frame
+	if (IsWindowInitialized())
+	{
+		// Clean FSlateApplication
+		FSlateApplication::Get().CloseAllWindowsImmediately();
+		// Clean FGlobalTabmanager
+		FGlobalTabmanager::Get()->CloseAllAreas();
+	}
+	// (Re-)create default main frame
+	CreateDefaultMainFrame(bStartImmersive, bStartPIE);
+}
+
 
 TSharedRef<SWidget> FMainFrameModule::MakeMainMenu(const TSharedPtr<FTabManager>& TabManager, const FName MenuName, FToolMenuContext& ToolMenuContext) const
 {
diff --git a/Engine/Source/Editor/MainFrame/Private/MainFrameModule.h b/Engine/Source/Editor/MainFrame/Private/MainFrameModule.h
--- a/Engine/Source/Editor/MainFrame/Private/MainFrameModule.h
+++ b/Engine/Source/Editor/MainFrame/Private/MainFrameModule.h
@@ -23,6 +23,7 @@ public:
 	// IMainFrameModule interface
 
 	virtual void CreateDefaultMainFrame( const bool bStartImmersive, const bool bStartPIE ) override;
+	virtual void RecreateDefaultMainFrame(const bool bStartImmersive, const bool bStartPIE) override;
 	virtual TSharedRef<SWidget> MakeMainMenu(const TSharedPtr<FTabManager>& TabManager, const FName MenuName, FToolMenuContext& ToolMenuContext) const override;
 	virtual TSharedRef<SWidget> MakeMainTabMenu(const TSharedPtr<FTabManager>& TabManager, const FName MenuName, FToolMenuContext& ToolMenuContext) const override;
 	virtual TSharedRef<SWidget> MakeDeveloperTools( const TArray<FMainFrameDeveloperTool>& AdditionalTools ) const override;
diff --git a/Engine/Source/Editor/MainFrame/Private/Menus/LayoutsMenu.cpp b/Engine/Source/Editor/MainFrame/Private/Menus/LayoutsMenu.cpp
new file mode 100644
--- /dev/null
+++ b/Engine/Source/Editor/MainFrame/Private/Menus/LayoutsMenu.cpp
@@ -0,0 +1,814 @@
+// Copyright 1998-2019 Epic Games, Inc. All Rights Reserved.
+
+#if WITH_EDITOR
+
+#include "Menus/LayoutsMenu.h"
+// Runtime/Core
+#include "Containers/Array.h"
+#include "GenericPlatform/GenericPlatformFile.h"
+#include "GenericPlatform/GenericPlatformMath.h"
+#include "HAL/FileManagerGeneric.h"
+#include "Logging/MessageLog.h"
+#include "Templates/SharedPointer.h"
+// Runtime
+#include "Framework/Commands/UICommandInfo.h"
+#include "Framework/Notifications/NotificationManager.h"
+#include "Misc/FileHelper.h"
+#include "Widgets/Notifications/SNotificationList.h"
+// Developer
+#include "IDesktopPlatform.h"
+#include "DesktopPlatformModule.h"
+#include "ToolMenus.h"
+// Editor
+#include "Classes/EditorStyleSettings.h"
+#include "Dialogs/Dialogs.h"
+#include "Editor/EditorPerProjectUserSettings.h"
+#include "Frame/MainFrameActions.h"
+#include "LevelViewportActions.h"
+#include "UnrealEdGlobals.h"
+#include "UnrealEdMisc.h"
+
+#define LOCTEXT_NAMESPACE "MainFrameActions"
+
+DEFINE_LOG_CATEGORY_STATIC(LogLayoutsMenu, Fatal, All);
+
+// Get LayoutsDirectory path
+FString CreateAndGetDefaultLayoutDirInternal()
+{
+	// Get LayoutsDirectory path
+	const FString LayoutsDirectory = FPaths::EngineDefaultLayoutDir();
+	// If the directory does not exist, create it (but it will not have saved Layouts inside)
+	IPlatformFile& PlatformFile = FPlatformFileManager::Get().GetPlatformFile();
+	if (!PlatformFile.DirectoryExists(*LayoutsDirectory))
+	{
+		PlatformFile.CreateDirectory(*LayoutsDirectory);
+	}
+	// Return result
+	return LayoutsDirectory;
+}
+
+FString CreateAndGetUserLayoutDirInternal()
+{
+	// Get UserLayoutsDirectory path
+	const FString UserLayoutsDirectory = FPaths::EngineUserLayoutDir();
+	// If the directory does not exist, create it (but it will not have saved Layouts inside)
+	IPlatformFile& PlatformFile = FPlatformFileManager::Get().GetPlatformFile();
+	if (!PlatformFile.DirectoryExists(*UserLayoutsDirectory))
+	{
+		PlatformFile.CreateDirectory(*UserLayoutsDirectory);
+	}
+	// Return result
+	return UserLayoutsDirectory;
+}
+
+TArray<FString> GetIniFilesInFolderInternal(const FString& InStringDirectory)
+{
+	// Find all ini files in folder
+	TArray<FString> LayoutIniFileNames;
+	const FString LayoutIniFilePaths = FPaths::Combine(*InStringDirectory, TEXT("*.ini"));
+	FFileManagerGeneric::Get().FindFiles(LayoutIniFileNames, *LayoutIniFilePaths, true, false);
+	return LayoutIniFileNames;
+}
+
+bool TrySaveLayoutOrWarnInternal(const FString& InSourceFilePath, const FString& InTargetFilePath, const FText& InWhatIsThis, const bool bCleanLayoutNameAndDescriptionFieldsIfNoSameValues, const bool bShouldAskBeforeCleaningLayoutNameAndDescriptionFields = false)
+{
+	// If desired, ask user whether to keep the LayoutName and LayoutDescription fields
+	bool bCleanLayoutNameAndDescriptionFields = false;
+	// If we are checking whether to clean the fields, we only want to maintain them if we are saving the file into an existing file that already has the same field values
+	if (bCleanLayoutNameAndDescriptionFieldsIfNoSameValues)
+	{
+		GConfig->UnloadFile(InSourceFilePath); // We must re-read it to avoid the Editor to use a previously cached name and description
+		const FText LayoutNameSource = FLayoutSaveRestore::LoadSectionFromConfig(InSourceFilePath, "LayoutName");
+		const FText LayoutDescriptionSource = FLayoutSaveRestore::LoadSectionFromConfig(InSourceFilePath, "LayoutDescription");
+		GConfig->UnloadFile(InTargetFilePath); // We must re-read it to avoid the Editor to use a previously cached name and description
+		const FText LayoutNameTarget = FLayoutSaveRestore::LoadSectionFromConfig(InTargetFilePath, "LayoutName");
+		const FText LayoutDescriptionTarget = FLayoutSaveRestore::LoadSectionFromConfig(InTargetFilePath, "LayoutDescription");
+		// The output target exists (overriding)
+		// These fields are not empty in source
+		if (!LayoutNameSource.IsEmpty() || !LayoutDescriptionSource.IsEmpty())
+		{
+			// These fields are different than the ones in target
+			if ((LayoutNameSource.ToString() != LayoutNameTarget.ToString()) || (LayoutDescriptionSource.ToString() != LayoutDescriptionTarget.ToString()))
+			{
+				bCleanLayoutNameAndDescriptionFields = true;
+				// We should clean the layout name and description fields, but ask user first
+				if (bShouldAskBeforeCleaningLayoutNameAndDescriptionFields)
+				{
+					// Open Dialog
+					const FText TextTitle = LOCTEXT("OverrideLayoutNameAndDescriptionFieldBodyTitle", "Clean UI Layout name and description fields");
+					const FText TextBody = FText::Format(
+						LOCTEXT("OverrideLayoutNameAndDescriptionFieldBody",
+							"You are saving a layout that contains a custom layout name and/or description. Do you also want to copy these 2 properties?\n"
+							" - Current layout name: {0}\n"
+							" - Current layout description: {1}\n\n"
+							"If you select \"Yes\", the displayed name and description of the original layout customization will also be copied into the new configuration file.\n\n"
+							"If you select \"No\", these fields will be emptied.\n\n"
+							"If you are not sure, select \"Yes\" if you are exporting the layout configuration without making any changes, or \"No\" if you have made or plan to make changes to the layout.\n\n"
+						),
+						LayoutNameSource, LayoutDescriptionSource);
+					const int AppReturnType = OpenMsgDlgInt(EAppMsgType::YesNoCancel, TextBody, TextTitle);
+					// Handle user answers
+					if (AppReturnType == EAppReturnType::Yes)
+					{
+						bCleanLayoutNameAndDescriptionFields = false;
+					}
+					else if (AppReturnType == EAppReturnType::No)
+					{
+						bCleanLayoutNameAndDescriptionFields = true;
+					}
+					else if (AppReturnType == EAppReturnType::Cancel)
+					{
+						return false;
+					}
+				}
+			}
+		}
+	}
+	// Copy: Replace main layout with desired one
+	const bool bShouldReplace = true;
+	const bool bCopyEvenIfReadOnly = true;
+	const bool bCopyAttributes = true;
+	if (COPY_Fail == IFileManager::Get().Copy(*InTargetFilePath, *InSourceFilePath, bShouldReplace, bCopyEvenIfReadOnly, bCopyAttributes))
+	{
+		FMessageLog EditorErrors("EditorErrors");
+		FText TextBody;
+		FFormatNamedArguments Arguments;
+		Arguments.Add(TEXT("WhatIs"), InWhatIsThis);
+		// Source does not exist
+		if (!FPaths::FileExists(InSourceFilePath))
+		{
+			Arguments.Add(TEXT("FileName"), FText::FromString(FPaths::ConvertRelativePathToFull(InSourceFilePath)));
+			TextBody = FText::Format(LOCTEXT("UnsuccessfulSave_NoExist_Notification", "Unsuccessful {WhatIs}, the desired file does not exist:\n{FileName}"), Arguments);
+			EditorErrors.Warning(TextBody);
+		}
+		// Target is read-only
+		else if (IFileManager::Get().IsReadOnly(*InTargetFilePath))
+		{
+			Arguments.Add(TEXT("FileName"), FText::FromString(FPaths::ConvertRelativePathToFull(InTargetFilePath)));
+			TextBody = FText::Format(LOCTEXT("UnsuccessfulSave_ReadOnly_Notification", "Unsuccessful {WhatIs}, the target file path is read-only\n{FileName}"), Arguments);
+			EditorErrors.Warning(TextBody);
+		}
+		// Target and source are the same
+		else if (FPaths::ConvertRelativePathToFull(InTargetFilePath) == FPaths::ConvertRelativePathToFull(InSourceFilePath))
+		{
+			Arguments.Add(TEXT("SourceFileName"), FText::FromString(FPaths::ConvertRelativePathToFull(InSourceFilePath)));
+			Arguments.Add(TEXT("FinalFileName"), FText::FromString(FPaths::ConvertRelativePathToFull(InTargetFilePath)));
+			TextBody = FText::Format(LOCTEXT("UnsuccessfulSave_Fallback_Notification",
+				"Unsuccessful {WhatIs}, target and source layout file paths are the same ({SourceFileName})!\nAre you trying to import or replace a file that is already in the layouts folder? If so, remove the current file first."), Arguments);
+			EditorErrors.Warning(TextBody);
+		}
+		// We don't specifically know why it failed, this is a fallback
+		else
+		{
+			Arguments.Add(TEXT("SourceFileName"), FText::FromString(FPaths::ConvertRelativePathToFull(InSourceFilePath)));
+			Arguments.Add(TEXT("FinalFileName"), FText::FromString(FPaths::ConvertRelativePathToFull(InTargetFilePath)));
+			TextBody = FText::Format(LOCTEXT("UnsuccessfulSave_Fallback_Notification", "Unsuccessful {WhatIs} of:\n{SourceFileName}\nto\n{FinalFileName}"), Arguments);
+			EditorErrors.Warning(TextBody);
+		}
+		EditorErrors.Notify(LOCTEXT("LoadUnsuccessful_Title", "Load Unsuccessful!"));
+		// Show reason
+		const FText TextTitle = LOCTEXT("UnsuccessfulCopyHeader", "Unsuccessful copy!");
+		OpenMsgDlgInt(EAppMsgType::Ok, TextBody, TextTitle);
+		// Return
+		return false;
+	}
+	// Copy successful
+	else
+	{
+		// Clean Layout Name and Description fields
+		// We copy twice to make sure we can copy.
+		// Problem if we only copied once: If the copy fails, the current EditorLayout.ini would be modified and no longer matches the previous one.
+		// The ini file should only be modified if it has been successfully copied to the new (and modified) INI file.
+		if (bCleanLayoutNameAndDescriptionFields)
+		{
+			// Update fields
+			FLayoutSaveRestore::SaveSectionToConfig(GEditorLayoutIni, "LayoutName", FText::FromString(""));
+			FLayoutSaveRestore::SaveSectionToConfig(GEditorLayoutIni, "LayoutDescription", FText::FromString(""));
+			// Flush file
+			const bool bRead = true;
+			GConfig->Flush(bRead, GEditorLayoutIni);
+			// Re-copy file
+			if (FPaths::ConvertRelativePathToFull(InTargetFilePath) != FPaths::ConvertRelativePathToFull(GEditorLayoutIni))
+			{
+				IFileManager::Get().Copy(*InTargetFilePath, *GEditorLayoutIni, bShouldReplace, bCopyEvenIfReadOnly, bCopyAttributes);
+			}
+		}
+		// Unload target file so it can be re-read into cache properly the next time it is used
+		GConfig->UnloadFile(InTargetFilePath); // We must re-read it to avoid the Editor to use a previously cached name and description
+		// Return successful copy message
+		return true;
+	}
+}
+
+// Name into display text
+FText GetDisplayTextInternal(const FString& InString)
+{
+	const FText DisplayNameText = FText::FromString(FPaths::GetBaseFilename(*InString));
+	const bool bIsBool = false;
+	const FText DisplayName = FText::FromString(FName::NameToDisplayString(DisplayNameText.ToString(), bIsBool));
+	return DisplayName;
+}
+FText GetTooltipTextInternal(const FString& InKindOfFile, const FText& InDisplayName, const FString& InLayoutFilePath)
+{
+	return FText::FromString(InKindOfFile + FString(" name:\n") + InDisplayName.ToString() + FString("\n\nFull file path:\n") + InLayoutFilePath);
+}
+
+void DisplayLayoutsInternal(FToolMenuSection& InSection, const TArray<TSharedPtr<FUICommandInfo>>& InXLayoutCommands, const TArray<FString> InLayoutIniFileNames, const FString InLayoutsDirectory)
+{
+	// If there are Layout ini files, read them
+	for (int32 LayoutIndex = 0; LayoutIndex < InLayoutIniFileNames.Num() && LayoutIndex < InXLayoutCommands.Num(); ++LayoutIndex)
+	{
+		const FString LayoutFilePath = IFileManager::Get().ConvertToAbsolutePathForExternalAppForRead(*FPaths::Combine(InLayoutsDirectory, InLayoutIniFileNames[LayoutIndex]));
+		// Make sure it is a layout file
+		GConfig->UnloadFile(LayoutFilePath); // We must re-read it to avoid the Editor to use a previously cached name and description
+		if (FLayoutSaveRestore::IsValidConfig(LayoutFilePath))
+		{
+			// Read and display localization name from INI file
+			const FText LayoutName = FLayoutSaveRestore::LoadSectionFromConfig(LayoutFilePath, "LayoutName");
+			const FText LayoutDescription = FLayoutSaveRestore::LoadSectionFromConfig(LayoutFilePath, "LayoutDescription");
+			// If no localization name, then display the file name
+			const FText DisplayName = (!LayoutName.IsEmpty() ? LayoutName : GetDisplayTextInternal(InLayoutIniFileNames[LayoutIndex]));
+			const FText Tooltip = (!LayoutDescription.IsEmpty() ? LayoutDescription : GetTooltipTextInternal("Default layout", DisplayName, LayoutFilePath));
+			InSection.AddMenuEntry(InXLayoutCommands[LayoutIndex], DisplayName, Tooltip);
+		}
+	}
+}
+
+void MakeXLayoutsMenuInternal(UToolMenu* InToolMenu, const TArray<TSharedPtr<FUICommandInfo>>& InXLayoutCommands, const TArray<TSharedPtr<FUICommandInfo>>& InXUserLayoutCommands, const bool bDisplayDefaultLayouts)
+{
+	// Update GEditorLayoutIni file. Otherwise, we could not track the changes the user did since the layout was loaded
+	FLayoutsMenuSave::SaveLayout();
+	// UE Default Layouts
+	if (bDisplayDefaultLayouts)
+	{
+		FToolMenuSection& Section = InToolMenu->AddSection("DefaultLayouts", LOCTEXT("DefaultLayoutsHeading", "Default Layouts"));
+		// Get LayoutsDirectory path
+		const FString LayoutsDirectory = CreateAndGetDefaultLayoutDirInternal();
+		// Get Layout init files
+		const TArray<FString> LayoutIniFileNames = GetIniFilesInFolderInternal(LayoutsDirectory);
+		// If there are user Layout ini files, read and display them
+		DisplayLayoutsInternal(Section, InXLayoutCommands, LayoutIniFileNames, LayoutsDirectory);
+	}
+	// User Layouts
+	{
+		FToolMenuSection& Section = InToolMenu->AddSection("UserDefaultLayouts", LOCTEXT("UserDefaultLayoutsHeading", "User Layouts"));
+		// (Create if it does not exist and) Get UserLayoutsDirectory path
+		const FString UserLayoutsDirectory = CreateAndGetUserLayoutDirInternal();
+		// Get User Layout init files
+		const TArray<FString> UserLayoutIniFileNames = GetIniFilesInFolderInternal(UserLayoutsDirectory);
+		// If there are user Layout ini files, read and display them
+		DisplayLayoutsInternal(Section, InXUserLayoutCommands, UserLayoutIniFileNames, UserLayoutsDirectory);
+	}
+}
+
+// All can be read
+/**
+ * Static
+ * Checks if the selected layout can be read (e.g., when loading layouts).
+ * @param	InLayoutIndex  Index from the selected layout.
+ * @return true if the selected layout can be read.
+ */
+bool CanChooseLayoutWhenRead(const int32 InLayoutIndex)
+{
+	return true;
+}
+/**
+ * Static
+ * Checks if the selected user-created layout can be read (e.g., when loading user layouts).
+ * @param	InLayoutIndex  Index from the selected user-created layout.
+ * @return true if the selected user-created layout can be read.
+ */
+bool CanChooseUserLayoutWhenRead(const int32 InLayoutIndex)
+{
+	return true;
+}
+// Only the layouts created by the user can be modified
+/**
+ * Static
+ * Checks if the selected layout can be modified (e.g., when overriding or removing layouts).
+ * @param	InLayoutIndex  Index from the selected layout.
+ * @return true if the selected layout can be modified/removed.
+ */
+bool CanChooseLayoutWhenWrite(const int32 InLayoutIndex)
+{
+	return false;
+}
+/**
+ * Static
+ * Checks if the selected user-created layout can be modified (e.g., when overriding or removing user layouts).
+ * @param	InLayoutIndex  Index from the selected user-created layout.
+ * @return true if the selected user-created layout can be modified/removed.
+ */
+bool CanChooseUserLayoutWhenWrite(const int32 InLayoutIndex)
+{
+	return true;
+}
+
+
+
+void FLayoutsMenuLoad::MakeLoadLayoutsMenu(UToolMenu* InToolMenu)
+{
+	// MakeLoadLayoutsMenu
+	const bool bDisplayDefaultLayouts = true;
+	MakeXLayoutsMenuInternal(InToolMenu, FMainFrameCommands::Get().MainFrameLayoutCommands.LoadLayoutCommands, FMainFrameCommands::Get().MainFrameLayoutCommands.LoadUserLayoutCommands, bDisplayDefaultLayouts);
+
+	// Additional sections
+	if (GetDefault<UEditorStyleSettings>()->bEnableUserEditorLayoutManagement)
+	{
+		FToolMenuSection& Section = InToolMenu->FindOrAddSection("UserDefaultLayouts");
+
+		// Separator
+		if (FLayoutsMenuBase::IsThereUserLayouts())
+		{
+			Section.AddMenuSeparator("AdditionalSectionsSeparator");
+		}
+
+		// Import...
+		{
+			Section.AddMenuEntry(FMainFrameCommands::Get().MainFrameLayoutCommands.ImportLayout);
+		}
+	}
+}
+
+bool FLayoutsMenuLoad::CanLoadChooseLayout(const int32 InLayoutIndex)
+{
+	if (GEngine && GEngine->GameViewport)
+	{
+		return false;
+	}
+	return !FLayoutsMenuBase::IsLayoutChecked(InLayoutIndex) && CanChooseLayoutWhenRead(InLayoutIndex);
+}
+bool FLayoutsMenuLoad::CanLoadChooseUserLayout(const int32 InLayoutIndex)
+{
+	if (GEngine && GEngine->GameViewport)
+	{
+		return false;
+	}
+	return !FLayoutsMenuBase::IsUserLayoutChecked(InLayoutIndex) && CanChooseUserLayoutWhenRead(InLayoutIndex);
+}
+
+void FLayoutsMenuLoad::ReloadCurrentLayout()
+{
+	// Editor is reset on-the-fly
+	FUnrealEdMisc::Get().AllowSavingLayoutOnClose(false);
+	EditorReinit();
+	FUnrealEdMisc::Get().AllowSavingLayoutOnClose(true);
+}
+
+void FLayoutsMenuLoad::LoadLayout(const FString& InLayoutPath)
+{
+	// Replace main layout with desired one
+	const FString& SourceFilePath = InLayoutPath;
+	const FString& TargetFilePath = GEditorLayoutIni;
+	const bool bCleanLayoutNameAndDescriptionFieldsIfNoSameValues = false;
+	const bool SucessfullySaved = TrySaveLayoutOrWarnInternal(SourceFilePath, TargetFilePath, LOCTEXT("LoadLayoutText", "layout load"), bCleanLayoutNameAndDescriptionFieldsIfNoSameValues);
+	// Reload current layout
+	if (SucessfullySaved)
+	{
+		FLayoutsMenuLoad::ReloadCurrentLayout();
+	}
+}
+
+void FLayoutsMenuLoad::LoadLayout(const int32 InLayoutIndex)
+{
+	// Replace main layout with desired one, reset layout & restart Editor
+	LoadLayout(FLayoutsMenuBase::GetLayout(InLayoutIndex));
+}
+
+void FLayoutsMenuLoad::LoadUserLayout(const int32 InLayoutIndex)
+{
+	// Replace main layout with desired one, reset layout & restart Editor
+	LoadLayout(FLayoutsMenuBase::GetUserLayout(InLayoutIndex));
+}
+
+void FLayoutsMenuLoad::ImportLayout()
+{
+	// Import the user-selected layout configuration files and load one of them
+	IDesktopPlatform* DesktopPlatform = FDesktopPlatformModule::Get();
+	if (DesktopPlatform)
+	{
+		// Open File Dialog so user can select his/her desired layout configuration files
+		TArray<FString> LayoutFilePaths;
+		const FString LastDirectory = FPaths::ProjectContentDir();
+		const FString DefaultDirectory = LastDirectory;
+		const FString DefaultFile = "";
+		const bool bWereFilesSelected = DesktopPlatform->OpenFileDialog(
+			FSlateApplication::Get().FindBestParentWindowHandleForDialogs(nullptr),
+			FString("Import a Layout Configuration File"),
+			DefaultDirectory,
+			DefaultFile,
+			TEXT("Layout configuration files|*.ini|"),
+			EFileDialogFlags::Multiple, //EFileDialogFlags::None, // Allow/Avoid multiple file selection
+			LayoutFilePaths
+		);
+		// If file(s) selected, copy them into the user layouts directory and load one of them
+		if (bWereFilesSelected && LayoutFilePaths.Num() > 0)
+		{
+			// (Create if it does not exist and) Get UserLayoutsDirectory path
+			const FString UserLayoutsDirectory = CreateAndGetUserLayoutDirInternal();
+			// Iterate over selected layout ini files
+			FString FirstGoodLayoutFile = "";
+			const FText TrySaveLayoutOrWarnInternalText = LOCTEXT("ImportLayoutText", "layout import");
+			for (const FString& LayoutFilePath : LayoutFilePaths)
+			{
+				// If file is a layout file, import it
+				GConfig->UnloadFile(LayoutFilePath); // We must re-read it to avoid the Editor to use a previously cached name and description
+				if (FLayoutSaveRestore::IsValidConfig(LayoutFilePath))
+				{
+					if (FirstGoodLayoutFile == "")
+					{
+						FirstGoodLayoutFile = LayoutFilePath;
+					}
+					// Save in the user layout folder
+					const FString& SourceFilePath = LayoutFilePath;
+					const FString& TargetFilePath = FPaths::Combine(FPaths::GetPath(UserLayoutsDirectory), FPaths::GetCleanFilename(LayoutFilePath));
+					const bool bCleanLayoutNameAndDescriptionFieldsIfNoSameValues = false;
+					TrySaveLayoutOrWarnInternal(SourceFilePath, TargetFilePath, TrySaveLayoutOrWarnInternalText, bCleanLayoutNameAndDescriptionFieldsIfNoSameValues);
+				}
+				// File is not a layout file, warn the user
+				else
+				{
+					FFormatNamedArguments Arguments;
+					Arguments.Add(TEXT("FileName"), FText::FromString(FPaths::ConvertRelativePathToFull(LayoutFilePath)));
+					const FText TextBody = FText::Format(LOCTEXT("UnsuccessfulImportBody", "Unsuccessful import, {FileName} is not a layout configuration file!"), Arguments);
+					const FText TextTitle = LOCTEXT("UnsuccessfulImportHeader", "Unsuccessful Import!");
+					OpenMsgDlgInt(EAppMsgType::Ok, TextBody, TextTitle);
+				}
+			}
+			// If PIE running, do not reload current layout
+			if (GEngine && GEngine->GameViewport)
+			{
+				const FText TextBody = LOCTEXT("SuccessfulImportBody",
+					"The layout(s) were successfully imported into the \"User Layouts\" section. However, no layout has been loaded into your current Unreal Editor UI because PIE is currently running. In order to do so, you must stop PIE and then load the layout from the \"User Layouts\" section.");
+				const FText TextTitle = LOCTEXT("UnsuccessfulImportHeader", "Successful Import!");
+				OpenMsgDlgInt(EAppMsgType::Ok, TextBody, TextTitle);
+				return;
+			}
+			// Replace current layout with first one
+			if (FirstGoodLayoutFile != "")
+			{
+				const FString& SourceFilePath = FirstGoodLayoutFile;
+				const FString& TargetFilePath = GEditorLayoutIni;
+				const bool bCleanLayoutNameAndDescriptionFieldsIfNoSameValues = false;
+				const bool SucessfullySaved = TrySaveLayoutOrWarnInternal(SourceFilePath, TargetFilePath, TrySaveLayoutOrWarnInternalText, bCleanLayoutNameAndDescriptionFieldsIfNoSameValues);
+				// Reload current layout
+				if (SucessfullySaved)
+				{
+					ReloadCurrentLayout();
+				}
+			}
+		}
+	}
+}
+
+
+
+void SaveExportLayoutCommon(const FString& InDefaultDirectory, const bool bMustBeSavedInDefaultDirectory, const FText& InWhatIsThis, const bool bShouldAskBeforeCleaningLayoutNameAndDescriptionFields)
+{
+	// Export/SaveAs the user-selected layout configuration files and load one of them
+	IDesktopPlatform* DesktopPlatform = FDesktopPlatformModule::Get();
+	if (DesktopPlatform)
+	{
+		// Open the "save file" dialog so user can save his/her layout configuration file
+		TArray<FString> LayoutFilePaths;
+		const FString DialogTitle = "";
+		const FString DefaultFile = "";
+		const bool bWereFileSelected = DesktopPlatform->SaveFileDialog(
+			FSlateApplication::Get().FindBestParentWindowHandleForDialogs(nullptr),
+			FString("Export a Layout Configuration File"),
+			InDefaultDirectory,
+			DefaultFile,
+			TEXT("Layout configuration files|*.ini|"),
+			EFileDialogFlags::None, //EFileDialogFlags::Multiple, // Allow/Avoid multiple file selection
+			LayoutFilePaths
+		);
+		// If file(s) selected, copy them into the user layouts directory and load one of them
+		if (bWereFileSelected && LayoutFilePaths.Num() > 0)
+		{
+			// Iterate over selected layout ini files
+			FString FirstGoodLayoutFile = "";
+			for (const FString& LayoutFilePath : LayoutFilePaths)
+			{
+				// If writing in the right folder
+				const FString LayoutFilePathAbsolute = IFileManager::Get().ConvertToAbsolutePathForExternalAppForRead(*FPaths::GetPath(LayoutFilePath));
+				const FString DefaultDirectoryAbsolute = IFileManager::Get().ConvertToAbsolutePathForExternalAppForRead(*FPaths::GetPath(InDefaultDirectory));
+				if (!bMustBeSavedInDefaultDirectory || (LayoutFilePathAbsolute == DefaultDirectoryAbsolute))
+				{
+					// Save in the user layout folder
+					const FString& SourceFilePath = GEditorLayoutIni;
+					const FString& TargetFilePath = LayoutFilePath;
+					const bool bCleanLayoutNameAndDescriptionFieldsIfNoSameValues = true;
+					TrySaveLayoutOrWarnInternal(SourceFilePath, TargetFilePath, InWhatIsThis, bCleanLayoutNameAndDescriptionFieldsIfNoSameValues, bShouldAskBeforeCleaningLayoutNameAndDescriptionFields);
+				}
+				// If trying to write in a different folder (which is not allowed)
+				else
+				{
+					// Warn the user that the file will not be copied in there
+					OpenMsgDlgInt(
+						EAppMsgType::Ok,
+						FText::Format(
+							LOCTEXT("SaveAsFailedMsg",
+								"In order to save the layout and allow Unreal to use it, you must save it in the predefined folder:\n{0}\n\nNevertheless, you tried to save it in:\n{1}\n\n"
+								"If you simply wish to export a copy of the current configuration in {1} (e.g., to later copy it into a different machine), you could use the \"Export Layout...\""
+								" functionality. However, Unreal would not be able to load it until you import it with \"Import Layout...\"."),
+							FText::FromString(DefaultDirectoryAbsolute), FText::FromString(LayoutFilePathAbsolute)),
+						LOCTEXT("SaveAsFailedMsg_Title", "Save As Failed"));
+				}
+			}
+		}
+	}
+}
+
+void FLayoutsMenuSave::MakeSaveLayoutsMenu(UToolMenu* InToolMenu)
+{
+	if (GetDefault<UEditorStyleSettings>()->bEnableUserEditorLayoutManagement)
+	{
+		// MakeOverrideLayoutsMenu
+		const bool bDisplayDefaultLayouts = false;
+		MakeXLayoutsMenuInternal(InToolMenu, FMainFrameCommands::Get().MainFrameLayoutCommands.OverrideLayoutCommands, FMainFrameCommands::Get().MainFrameLayoutCommands.OverrideUserLayoutCommands, bDisplayDefaultLayouts);
+
+		// Additional sections
+		{
+			FToolMenuSection& Section = InToolMenu->FindOrAddSection("UserDefaultLayouts");
+
+			// Separator
+			if (FLayoutsMenuBase::IsThereUserLayouts())
+			{
+				Section.AddMenuSeparator("AdditionalSectionsSeparator");
+			}
+
+			// Save as...
+			{
+				Section.AddMenuEntry(FMainFrameCommands::Get().MainFrameLayoutCommands.SaveLayoutAs);
+			}
+
+			// Export...
+			{
+				Section.AddMenuEntry(FMainFrameCommands::Get().MainFrameLayoutCommands.ExportLayout);
+			}
+		}
+	}
+}
+
+bool FLayoutsMenuSave::CanSaveChooseLayout(const int32 InLayoutIndex)
+{
+	return !FLayoutsMenuBase::IsLayoutChecked(InLayoutIndex) && CanChooseLayoutWhenWrite(InLayoutIndex);
+}
+bool FLayoutsMenuSave::CanSaveChooseUserLayout(const int32 InLayoutIndex)
+{
+	return !FLayoutsMenuBase::IsUserLayoutChecked(InLayoutIndex) && CanChooseUserLayoutWhenWrite(InLayoutIndex);
+}
+
+void FLayoutsMenuSave::OverrideLayout(const int32 InLayoutIndex)
+{
+	// Default layouts should never be modified, so this function should never be called
+	UE_LOG(LogLayoutsMenu, Fatal, TEXT("Default layouts should never be modified, so this function should never be called."));
+}
+
+void FLayoutsMenuSave::OverrideUserLayout(const int32 InLayoutIndex)
+{
+	// (Create if it does not exist and) Get UserLayoutsDirectory path
+	const FString UserLayoutsDirectory = CreateAndGetUserLayoutDirInternal();
+	// Get User Layout init files
+	const TArray<FString> UserLayoutIniFileNames = GetIniFilesInFolderInternal(UserLayoutsDirectory);
+	const FString DesiredUserLayoutFullPath = FPaths::Combine(FPaths::GetPath(UserLayoutsDirectory), UserLayoutIniFileNames[InLayoutIndex]);
+	// Are you sure you want to do this?
+	const FText TextFileNameToRemove = FText::FromString(FPaths::GetBaseFilename(UserLayoutIniFileNames[InLayoutIndex]));
+	const FText TextBody = FText::Format(LOCTEXT("ActionOverrideLayoutMsg", "Are you sure you want to permanently override the layout profile \"{0}\" with the current layout profile? This action cannot be undone."), TextFileNameToRemove);
+	const FText TextTitle = FText::Format(LOCTEXT("OverrideUILayout_Title", "Override UI Layout \"{0}\""), TextFileNameToRemove);
+	if (EAppReturnType::Ok != OpenMsgDlgInt(EAppMsgType::OkCancel, TextBody, TextTitle))
+	{
+		return;
+	}
+	// Target and source files
+	const FString& SourceFilePath = GEditorLayoutIni;
+	const FString& TargetFilePath = DesiredUserLayoutFullPath;
+	// Update GEditorLayoutIni file
+	SaveLayout();
+	// Replace desired layout with current one
+	const bool bCleanLayoutNameAndDescriptionFieldsIfNoSameValues = true;
+	const bool bShouldAskBeforeCleaningLayoutNameAndDescriptionFields = false;
+	TrySaveLayoutOrWarnInternal(SourceFilePath, TargetFilePath, LOCTEXT("OverrideLayoutText", "layout override"), bCleanLayoutNameAndDescriptionFieldsIfNoSameValues);
+}
+
+void FLayoutsMenuSave::SaveLayout()
+{
+	// Save the layout into the Editor
+	FGlobalTabmanager::Get()->SaveAllVisualState();
+	// Write the saved layout to disk (if it has changed since the last time it was read/written)
+	// We must set bRead = true. Otherwise, FLayoutsMenuLoad::ReloadCurrentLayout() would reload the old config file (because it would be cached on memory)
+	const bool bRead = true;
+	GConfig->Flush(bRead, GEditorLayoutIni);
+}
+
+void FLayoutsMenuSave::SaveLayoutAs()
+{
+	// Update GEditorLayoutIni file
+	SaveLayout();
+	// Copy GEditorLayoutIni into desired file
+	const FString DefaultDirectory = CreateAndGetUserLayoutDirInternal();
+	const bool bMustBeSavedInDefaultDirectory = true;
+	const bool bShouldAskBeforeCleaningLayoutNameAndDescriptionFields = false;
+	SaveExportLayoutCommon(DefaultDirectory, bMustBeSavedInDefaultDirectory, LOCTEXT("SaveLayoutText", "layout save"), bShouldAskBeforeCleaningLayoutNameAndDescriptionFields);
+}
+
+void FLayoutsMenuSave::ExportLayout()
+{
+	// Update GEditorLayoutIni file
+	SaveLayout();
+	// Copy GEditorLayoutIni into desired file
+	const FString DefaultDirectory = FPaths::ProjectContentDir();
+	const bool bMustBeSavedInDefaultDirectory = false;
+	const bool bShouldAskBeforeCleaningLayoutNameAndDescriptionFields = true;
+	SaveExportLayoutCommon(DefaultDirectory, bMustBeSavedInDefaultDirectory, LOCTEXT("ExportLayoutText", "layout export"), bShouldAskBeforeCleaningLayoutNameAndDescriptionFields);
+}
+
+
+
+void FLayoutsMenuRemove::MakeRemoveLayoutsMenu(UToolMenu* InToolMenu)
+{
+	if (GetDefault<UEditorStyleSettings>()->bEnableUserEditorLayoutManagement)
+	{
+		// MakeRemoveLayoutsMenu
+		const bool bDisplayDefaultLayouts = false;
+		MakeXLayoutsMenuInternal(InToolMenu, FMainFrameCommands::Get().MainFrameLayoutCommands.RemoveLayoutCommands, FMainFrameCommands::Get().MainFrameLayoutCommands.RemoveUserLayoutCommands, bDisplayDefaultLayouts);
+
+		// Additional sections
+		{
+			FToolMenuSection& Section = InToolMenu->FindOrAddSection("UserDefaultLayouts");
+
+			// Separator
+			if (FLayoutsMenuBase::IsThereUserLayouts())
+			{
+				Section.AddMenuSeparator("AdditionalSectionsSeparator");
+			}
+
+			// Remove all
+			Section.AddMenuEntry(FMainFrameCommands::Get().MainFrameLayoutCommands.RemoveUserLayouts);
+		}
+	}
+}
+
+bool FLayoutsMenuRemove::CanRemoveChooseLayout(const int32 InLayoutIndex)
+{
+	return CanChooseLayoutWhenWrite(InLayoutIndex);
+}
+bool FLayoutsMenuRemove::CanRemoveChooseUserLayout(const int32 InLayoutIndex)
+{
+	return CanChooseUserLayoutWhenWrite(InLayoutIndex);
+}
+
+void FLayoutsMenuRemove::RemoveLayout(const int32 InLayoutIndex)
+{
+	// Default layouts should never be modified, so this function should never be called
+	UE_LOG(LogLayoutsMenu, Fatal, TEXT("Default layouts should never be modified, so this function should never be called."));
+}
+
+void FLayoutsMenuRemove::RemoveUserLayout(const int32 InLayoutIndex)
+{
+	// (Create if it does not exist and) Get UserLayoutsDirectory path
+	const FString UserLayoutsDirectory = CreateAndGetUserLayoutDirInternal();
+	// Get User Layout init files
+	const TArray<FString> UserLayoutIniFileNames = GetIniFilesInFolderInternal(UserLayoutsDirectory);
+	const FString DesiredUserLayoutFullPath = FPaths::Combine(FPaths::GetPath(UserLayoutsDirectory), UserLayoutIniFileNames[InLayoutIndex]);
+	// Are you sure you want to do this?
+	const FText TextFileNameToRemove = FText::FromString(FPaths::GetBaseFilename(UserLayoutIniFileNames[InLayoutIndex]));
+	const FText TextBody = FText::Format(LOCTEXT("ActionRemoveMsg", "Are you sure you want to permanently delete the layout profile \"{0}\"? This action cannot be undone."), TextFileNameToRemove);
+	const FText TextTitle = FText::Format(LOCTEXT("RemoveUILayout_Title", "Remove UI Layout \"{0}\""), TextFileNameToRemove);
+	if (EAppReturnType::Ok != OpenMsgDlgInt(EAppMsgType::OkCancel, TextBody, TextTitle))
+	{
+		return;
+	}
+	// Remove layout
+	FPlatformFileManager::Get().GetPlatformFile().DeleteFile(*DesiredUserLayoutFullPath);
+}
+
+int32 GetNumberLayoutFiles(const FString& InLayoutsDirectory)
+{
+	// Get Layout init files in desired directory
+	const TArray<FString> LayoutIniFileNames = GetIniFilesInFolderInternal(InLayoutsDirectory);
+	// Count how many layout files exist
+	int32 NumberLayoutFiles = 0;
+	for (const FString& LayoutIniFileName : LayoutIniFileNames)
+	{
+		const FString LayoutFilePath = FPaths::Combine(InLayoutsDirectory, LayoutIniFileName);
+		GConfig->UnloadFile(LayoutFilePath); // We must re-read it to avoid the Editor to use a previously cached name and description
+		if (FLayoutSaveRestore::IsValidConfig(LayoutFilePath))
+		{
+			++NumberLayoutFiles;
+		}
+	}
+	// Return result
+	return NumberLayoutFiles;
+}
+
+void FLayoutsMenuRemove::RemoveUserLayouts()
+{
+	// (Create if it does not exist and) Get UserLayoutsDirectory path
+	const FString UserLayoutsDirectory = CreateAndGetUserLayoutDirInternal();
+	// Count how many layout files exist
+	const int32 NumberUserLayoutFiles = GetNumberLayoutFiles(UserLayoutsDirectory);
+	// If files to remove, warn user and remove them all
+	if (NumberUserLayoutFiles > 0)
+	{
+		// Are you sure you want to do this?
+		const FText TextNumberUserLayoutFiles  = FText::FromString(FString::FromInt(NumberUserLayoutFiles));
+		const FText TextBody = FText::Format(LOCTEXT("ActionRemoveAllUserLayoutMsg", "Are you sure you want to permanently remove all the {0} layout profiles created by the user? This action cannot be undone."), TextNumberUserLayoutFiles);
+		const FText TextTitle = LOCTEXT("RemoveAllUserLayouts_Title", "Remove All User-Created Layouts");
+		if (EAppReturnType::Ok != OpenMsgDlgInt(EAppMsgType::OkCancel, TextBody, TextTitle))
+		{
+			return;
+		}
+		// Get User Layout init files
+		const TArray<FString> UserLayoutIniFileNames = GetIniFilesInFolderInternal(UserLayoutsDirectory);
+		// If there are user Layout ini files, read them
+		for (const FString& UserLayoutIniFileName : UserLayoutIniFileNames)
+		{
+			// Remove file if it is a layout
+			const FString LayoutFilePath = FPaths::Combine(UserLayoutsDirectory, UserLayoutIniFileName);
+			GConfig->UnloadFile(LayoutFilePath); // We must re-read it to avoid the Editor to use a previously cached name and description
+			if (FLayoutSaveRestore::IsValidConfig(LayoutFilePath))
+			{
+				FPlatformFileManager::Get().GetPlatformFile().DeleteFile(*LayoutFilePath);
+			}
+		}
+	}
+	// If no files to remove, warn user
+	else
+	{
+		// Show reason
+		const FText TextBody = LOCTEXT("UnsuccessfulRemoveLayoutBody", "There are no layout profile files created by the user, so none could be removed.");
+		const FText TextTitle = LOCTEXT("UnsuccessfulRemoveLayoutHeader", "Unsuccessful Remove All User Layouts!");
+		OpenMsgDlgInt(EAppMsgType::Ok, TextBody, TextTitle);
+	}
+}
+
+
+
+bool IsLayoutCheckedInternal(const FString& InLayoutFullPath)
+{
+	// Checked if same file. I.e.,
+		// 1. Same size
+		// 2. And same internal text
+	const bool bHaveSameSize = (IFileManager::Get().FileSize(*GEditorLayoutIni) == IFileManager::Get().FileSize(*InLayoutFullPath));
+	// Same size --> Same layout file?
+	if (bHaveSameSize)
+	{
+		// Read files and check whether they have the exact same text
+		FString StringGEditorLayoutIni;
+		FFileHelper::LoadFileToString(StringGEditorLayoutIni, *GEditorLayoutIni);
+		FString StringLayoutFullPath;
+		FFileHelper::LoadFileToString(StringLayoutFullPath, *InLayoutFullPath);
+		// (No) same text = (No) same layout file
+		return (StringGEditorLayoutIni == StringLayoutFullPath);
+	}
+	// No same size = No same layout file
+	else
+	{
+		return false;
+	}
+}
+
+FString FLayoutsMenuBase::GetLayout(const int32 InLayoutIndex)
+{
+	// Get LayoutsDirectory path, layout init files, and desired layout path
+	const FString LayoutsDirectory = CreateAndGetDefaultLayoutDirInternal();
+	const TArray<FString> LayoutIniFileNames = GetIniFilesInFolderInternal(LayoutsDirectory);
+	const FString DesiredLayoutFullPath = FPaths::Combine(FPaths::GetPath(LayoutsDirectory), LayoutIniFileNames[InLayoutIndex]);
+	// Return full path
+	return DesiredLayoutFullPath;
+}
+
+FString FLayoutsMenuBase::GetUserLayout(const int32 InLayoutIndex)
+{
+	// (Create if it does not exist and) Get UserLayoutsDirectory path, user layout init files, and desired user layout path
+	const FString UserLayoutsDirectory = CreateAndGetUserLayoutDirInternal();
+	const TArray<FString> UserLayoutIniFileNames = GetIniFilesInFolderInternal(UserLayoutsDirectory);
+	const FString DesiredUserLayoutFullPath = FPaths::Combine(FPaths::GetPath(UserLayoutsDirectory), UserLayoutIniFileNames[InLayoutIndex]);
+	// Return full path
+	return DesiredUserLayoutFullPath;
+}
+
+bool FLayoutsMenuBase::IsThereUserLayouts()
+{
+	// (Create if it does not exist and) Get UserLayoutsDirectory path
+	const FString UserLayoutsDirectory = CreateAndGetUserLayoutDirInternal();
+	// At least 1 user layout file?
+	return GetNumberLayoutFiles(UserLayoutsDirectory) > 0;
+}
+
+bool FLayoutsMenuBase::IsLayoutChecked(const int32 InLayoutIndex)
+{
+	// Check if the desired layout file matches the one currently loaded
+	return IsLayoutCheckedInternal(GetLayout(InLayoutIndex));
+}
+
+bool FLayoutsMenuBase::IsUserLayoutChecked(const int32 InLayoutIndex)
+{
+	// Check if the desired layout file matches the one currently loaded
+	return IsLayoutCheckedInternal(GetUserLayout(InLayoutIndex));
+}
+
+#undef LOCTEXT_NAMESPACE
+
+#endif
diff --git a/Engine/Source/Editor/MainFrame/Private/Menus/LayoutsMenu.h b/Engine/Source/Editor/MainFrame/Private/Menus/LayoutsMenu.h
new file mode 100644
--- /dev/null
+++ b/Engine/Source/Editor/MainFrame/Private/Menus/LayoutsMenu.h
@@ -0,0 +1,231 @@
+// Copyright 1998-2019 Epic Games, Inc. All Rights Reserved.
+
+#if WITH_EDITOR
+
+#pragma once
+
+class UToolMenu;
+
+/**
+ * Static load-related helper functions for populating the "Layouts" menu.
+ */
+class FLayoutsMenuLoad
+{
+public:
+	/**
+	 * Static
+	 * It creates the layout load selection menu.
+	 */
+	static void MakeLoadLayoutsMenu(UToolMenu* InToolMenu);
+
+	/**
+	 * Static
+	 * Checks if the load menu can choose the selected layout to load it.
+	 * @param	InLayoutIndex  Index from the selected layout.
+	 * @return true if the selected layout can be read.
+	 */
+	static bool CanLoadChooseLayout(const int32 InLayoutIndex);
+
+	/**
+	 * Static
+	 * Checks if the load menu can choose the selected user-created layout to load it.
+	 * @param	InLayoutIndex  Index from the selected user-created layout.
+	 * @return true if the selected user-created layout can be read.
+	 */
+	static bool CanLoadChooseUserLayout(const int32 InLayoutIndex);
+
+	/**
+	 * Static
+	 * It re-loads the current Editor UI Layout (from GEditorLayoutIni).
+	 * This function is used for many of the functions of FLayoutsMenuLoad.
+	 */
+	static void ReloadCurrentLayout();
+
+	/**
+	 * Static
+	 * Load the visual layout state of the editor from an existing layout profile ini file, given its file path.
+	 * @param InLayoutPath File path associated with the desired layout profile ini file to be read/written.
+	 */
+	static void LoadLayout(const FString& InLayoutPath);
+
+	/**
+	 * Static
+	 * Load the visual layout state of the editor from an existing layout profile ini file, that was previously created by the UE developers.
+	 * @param InLayoutIndex Index associated with the desired layout profile ini file to be read/written.
+	 */
+	static void LoadLayout(const int32 InLayoutIndex);
+
+	/**
+	 * Static
+	 * Load the visual layout state of the editor from an existing layout profile ini file, that was previously created by the user.
+	 * @param InLayoutIndex Index associated with the desired layout profile ini file to be read/written.
+	 */
+	static void LoadUserLayout(const int32 InLayoutIndex);
+
+	/**
+	 * Static
+	 * Import a visual layout state of the editor from a custom directory path and with a custom file name chosen by the user.
+	 * It saves it into the user layout folder, and loads it.
+	 */
+	static void ImportLayout();
+};
+
+/**
+ * Static save-related helper functions for populating the "Layouts" menu.
+ */
+class FLayoutsMenuSave
+{
+public:
+	/**
+	 * Static
+	 * It creates the layout save selection menu.
+	 */
+	static void MakeSaveLayoutsMenu(UToolMenu* InToolMenu);
+
+	/**
+	 * Static
+	 * Checks if the save menu can choose the selected layout to modify it.
+	 * @param	InLayoutIndex  Index from the selected layout.
+	 * @return true if the selected layout can be modified/removed.
+	 */
+	static bool CanSaveChooseLayout(const int32 InLayoutIndex);
+
+	/**
+	 * Static
+	 * Checks if the save menu can choose the selected user-created layout to modify it.
+	 * @param	InLayoutIndex  Index from the selected user-created layout.
+	 * @return true if the selected user-created layout can be modified/removed.
+	 */
+	static bool CanSaveChooseUserLayout(const int32 InLayoutIndex);
+
+	/**
+	 * Static
+	 * Override the visual layout state of the editor in an existing layout profile ini file, that was previously created by the UE developers.
+	 * @param InLayoutIndex Index associated with the desired layout profile ini file to be read/written.
+	 */
+	static void OverrideLayout(const int32 InLayoutIndex);
+
+	/**
+	 * Static
+	 * Override the visual layout state of the editor in an existing layout profile ini file, that was previously created by the user.
+	 * @param InLayoutIndex Index associated with the desired layout profile ini file to be read/written.
+	 */
+	static void OverrideUserLayout(const int32 InLayoutIndex);
+
+	/**
+	 * Static
+	 * Save the visual layout state of the editor (if changes to the layout have been made since the last time it was saved).
+	 * If no changes has been made to the layout, the file is not updated (given that it would not be required).
+	 * Any function that saves the layout (e.g., OverrideLayoutI, OverrideUserLayoutI, SaveLayoutAs, ExportLayout, etc.) should internally call this function.
+	*/
+	static void SaveLayout();
+
+	/**
+	 * Static
+	 * Save the visual layout state of the editor with a custom file name chosen by the user.
+	 */
+	static void SaveLayoutAs();
+
+	/**
+	 * Static
+	 * Export the visual layout state of the editor in a custom directory path and with a custom file name chosen by the user.
+	 */
+	static void ExportLayout();
+};
+
+/**
+ * Static remove-related helper functions for populating the "Layouts" menu.
+ */
+class FLayoutsMenuRemove
+{
+public:
+	/**
+	 * Static
+	 * It creates the layout remove selection menu.
+	 */
+	static void MakeRemoveLayoutsMenu(UToolMenu* InToolMenu);
+
+	/**
+	 * Static
+	 * Checks if the remove menu can choose the selected layout to remove it.
+	 * @param	InLayoutIndex  Index from the selected layout.
+	 * @return true if the selected layout can be modified/removed.
+	 */
+	static bool CanRemoveChooseLayout(const int32 InLayoutIndex);
+
+	/**
+	 * Static
+	 * Checks if the remove menu can choose the selected user-created layout to remove it.
+	 * @param	InLayoutIndex  Index from the selected user-created layout.
+	 * @return true if the selected user-created layout can be modified/removed.
+	 */
+	static bool CanRemoveChooseUserLayout(const int32 InLayoutIndex);
+
+	/**
+	 * Static
+	 * Remove an existing layout profile ini file, that was previously created by the UE developers.
+	 * @param InLayoutIndex Index associated with the desired layout profile ini file to be read/written.
+	 */
+	static void RemoveLayout(const int32 InLayoutIndex);
+
+	/**
+	 * Static
+	 * Remove an existing layout profile ini file, that was previously created by the UE developers.
+	 * @param InLayoutIndex Index associated with the desired layout profile ini file to be read/written.
+	 */
+	static void RemoveUserLayout(const int32 InLayoutIndex);
+
+	/**
+	 * Static
+	 * Remove all the layout customizations created by the user.
+	 */
+	static void RemoveUserLayouts();
+};
+
+/**
+ * Static helper functions for populating the "Layouts" menu.
+ */
+class FLayoutsMenuBase
+{
+public:
+	/**
+	 * Static
+	 * Get the full layout file path.
+	 * Helper function for LoadLayoutI, SaveLayoutI, and RemoveLayoutI.
+	 * @param InLayoutIndex Index associated with the desired layout profile ini file to be read/written.
+	 */
+	static FString GetLayout(const int32 InLayoutIndex);
+
+	/**
+	 * Static
+	 * Get the full user-created layout file path.
+	 * Helper function for LoadUserLayoutI, SaveUserLayoutI, and RemoveUserLayoutI.
+	 * @param InLayoutIndex Index associated with the desired user-created layout profile ini file to be read/written.
+	 */
+	static FString GetUserLayout(const int32 InLayoutIndex);
+
+	/**
+	 * Static
+	 * Checks whether there are user-created layouts.
+	 * @return true if there is at least a layout in the user layouts directory.
+	 */
+	static bool IsThereUserLayouts();
+
+	/**
+	 * Static
+	 * Checks which layout entry should be checked.
+	 * @param	InLayoutIndex  Index from the selected layout.
+	 * @return true if the menu entry should be checked.
+	 */
+	static bool IsLayoutChecked(const int32 InLayoutIndex);
+
+	/**
+	 * Static
+	 * Checks which user layout entry should be checked.
+	 * @param	InLayoutIndex  Index from the selected layout.
+	 * @return true if the menu entry should be checked.
+	 */
+	static bool IsUserLayoutChecked(const int32 InLayoutIndex);
+};
+
+#endif
diff --git a/Engine/Source/Editor/MainFrame/Private/Menus/MainMenu.cpp b/Engine/Source/Editor/MainFrame/Private/Menus/MainMenu.cpp
--- a/Engine/Source/Editor/MainFrame/Private/Menus/MainMenu.cpp
+++ b/Engine/Source/Editor/MainFrame/Private/Menus/MainMenu.cpp
@@ -18,6 +18,7 @@
 #include "Settings/EditorExperimentalSettings.h"
 #include "UnrealEdGlobals.h"
 #include "Frame/MainFrameActions.h"
+#include "Menus/LayoutsMenu.h"
 #include "Menus/PackageProjectMenu.h"
 #include "Menus/RecentProjectsMenu.h"
 #include "Menus/SettingsMenu.h"
@@ -209,6 +210,7 @@ void FMainMenu::RegisterWindowMenu()
 {
 	UToolMenu* Menu = UToolMenus::Get()->RegisterMenu("MainFrame.MainMenu.Window");
 
+	// Level Editor, General, and Testing sections
 	// Automatically populate tab spawners from TabManager
 	Menu->AddDynamicSection("TabManagerSection", FNewToolMenuDelegateLegacy::CreateLambda([](FMenuBuilder& InBuilder, UToolMenu* InData)
 	{
@@ -227,6 +229,7 @@ void FMainMenu::RegisterWindowMenu()
 		}
 	}));
 
+	// Project Launcher section
 	{
 		FToolMenuSection& Section = Menu->AddSection("WindowGlobalTabSpawners");
 		Section.AddMenuEntry(
@@ -238,11 +241,12 @@ void FMainMenu::RegisterWindowMenu()
 			);
 	}
 
+	// Experimental section
 	{
 		// This is a temporary home for the spawners of experimental features that must be explicitly enabled.
 		// When the feature becomes permanent and need not check a flag, register a nomad spawner for it in the proper WorkspaceMenu category
-		bool bLocalizationDashboard = GetDefault<UEditorExperimentalSettings>()->bEnableLocalizationDashboard;
-		bool bTranslationPicker = GetDefault<UEditorExperimentalSettings>()->bEnableTranslationPicker;
+		const bool bLocalizationDashboard = GetDefault<UEditorExperimentalSettings>()->bEnableLocalizationDashboard;
+		const bool bTranslationPicker = GetDefault<UEditorExperimentalSettings>()->bEnableTranslationPicker;
 
 		// Make sure at least one is enabled before creating the section
 		if (bLocalizationDashboard || bTranslationPicker)
@@ -276,11 +280,45 @@ void FMainMenu::RegisterWindowMenu()
 		}
 	}
 
+	// Layout section
 	{
 		FToolMenuSection& Section = Menu->AddSection("WindowLayout", NSLOCTEXT("MainAppMenu", "LayoutManagementHeader", "Layout"));
-		Section.AddMenuEntry(FMainFrameCommands::Get().ResetLayout);
-		Section.AddMenuEntry(FMainFrameCommands::Get().SaveLayout);
+		// Load Layout
+		Section.AddEntry(FToolMenuEntry::InitSubMenu(
+			"LoadLayout",
+			NSLOCTEXT("LayoutMenu", "LayoutLoadHeader", "Load Layout"),
+			NSLOCTEXT("LayoutMenu", "LoadLayoutsSubMenu_ToolTip", "Load a layout configuration from disk. If PIE is running, most options will be disabled"),
+			FNewToolMenuDelegate::CreateStatic(&FLayoutsMenuLoad::MakeLoadLayoutsMenu)
+		));
+		// Save and Remove Layout
+		// Opposite to "Load Layout", Save and Remove are dynamic, i.e., they can be enabled/removed depending on the value of
+		// the variable GetDefault<UEditorStyleSettings>()->bEnableUserEditorLayoutManagement
+		Section.AddDynamicEntry("OverrideAndRemoveLayout", FNewToolMenuSectionDelegate::CreateLambda([](FToolMenuSection& InSection)
+		{
+			if (GetDefault<UEditorStyleSettings>()->bEnableUserEditorLayoutManagement)
+			{
+				// Save Layout
+				InSection.AddEntry(FToolMenuEntry::InitSubMenu(
+					"OverrideLayout",
+					NSLOCTEXT("LayoutMenu", "OverrideLayoutsSubMenu", "Save Layout"),
+					NSLOCTEXT("LayoutMenu", "OverrideLayoutsSubMenu_ToolTip", "Save your current layout configuration on disk"),
+					FNewToolMenuDelegate::CreateStatic(&FLayoutsMenuSave::MakeSaveLayoutsMenu)
+				));
+				// Remove Layout
+				InSection.AddEntry(FToolMenuEntry::InitSubMenu(
+					"RemoveLayout",
+					NSLOCTEXT("LayoutMenu", "RemoveLayoutsSubMenu", "Remove Layout"),
+					NSLOCTEXT("LayoutMenu", "RemoveLayoutsSubMenu_ToolTip", "Remove a layout configuration from disk"),
+					FNewToolMenuDelegate::CreateStatic(&FLayoutsMenuRemove::MakeRemoveLayoutsMenu)
+				));
+			}
+		}));
+
+		// Enable Fullscreen section
 #if !PLATFORM_MAC && !PLATFORM_LINUX // On Mac/Linux windowed fullscreen mode in the editor is currently unavailable
+		// Separator
+		Section.AddMenuSeparator("FullscreenSeparator");
+		// Fullscreen
 		Section.AddMenuEntry(FMainFrameCommands::Get().ToggleFullscreen);
 #endif
 	}
diff --git a/Engine/Source/Editor/MainFrame/Public/Interfaces/IMainFrameModule.h b/Engine/Source/Editor/MainFrame/Public/Interfaces/IMainFrameModule.h
--- a/Engine/Source/Editor/MainFrame/Public/Interfaces/IMainFrameModule.h
+++ b/Engine/Source/Editor/MainFrame/Public/Interfaces/IMainFrameModule.h
@@ -43,6 +43,15 @@ public:
 	 */
 	virtual void CreateDefaultMainFrame( const bool bStartImmersive, const bool bStartPIE ) = 0;
 
+	/**
+	 * Recreates the default editor main frame.
+	 * I.e., if CreateDefaultMainFrame or RecreateDefaultMainFrame were already called, it would clean the previous default main frame and create it again.
+	 *
+	 * @param bStartImmersive True to force a main frame viewport into immersive mode
+	 * @param bStartPIE True to start a PIE session right away
+	 */
+	virtual void RecreateDefaultMainFrame(const bool bStartImmersive, const bool bStartPIE) = 0;
+
 	/**
 	 * Generates a menu that includes application global commands, such as "Save All", "Exit", etc.  If you're building
 	 * a menu for your tab, you should call this function to create your menu, passing in an extender object to add your
diff --git a/Engine/Source/Editor/UnrealEd/Private/Toolkits/AssetEditorToolkit.cpp b/Engine/Source/Editor/UnrealEd/Private/Toolkits/AssetEditorToolkit.cpp
--- a/Engine/Source/Editor/UnrealEd/Private/Toolkits/AssetEditorToolkit.cpp
+++ b/Engine/Source/Editor/UnrealEd/Private/Toolkits/AssetEditorToolkit.cpp
@@ -1092,7 +1092,7 @@ void FAssetEditorToolkit::RestoreFromLayout(const TSharedRef<FTabManager::FLayou
 	if (HostWidget.Get() != NULL)
 	{
 		// Save the old layout
-		FLayoutSaveRestore::SaveToConfig(GEditorIni, TabManager->PersistLayout());
+		FLayoutSaveRestore::SaveToConfig(GEditorLayoutIni, TabManager->PersistLayout());
 
 		// Load the potentially previously saved new layout
 		TSharedRef<FTabManager::FLayout> UserConfiguredNewLayout = FLayoutSaveRestore::LoadFromConfig(GEditorLayoutIni, NewLayout);
diff --git a/Engine/Source/Editor/UnrealEd/Private/UnrealEdGlobals.cpp b/Engine/Source/Editor/UnrealEd/Private/UnrealEdGlobals.cpp
--- a/Engine/Source/Editor/UnrealEd/Private/UnrealEdGlobals.cpp
+++ b/Engine/Source/Editor/UnrealEd/Private/UnrealEdGlobals.cpp
@@ -61,6 +61,22 @@ const FString GetEditorResourcesDir()
 	return FPaths::Combine( FPlatformProcess::BaseDir(), *FPaths::EngineContentDir(), TEXT("Editor/") );
 }
 
+void CheckAndMaybeGoToVRModeInternal(const bool bIsImmersive)
+{
+	// Go straight to VR mode if we were asked to
+	{
+		if (!bIsImmersive && FParse::Param(FCommandLine::Get(), TEXT("VREditor")))
+		{
+			IVREditorModule& VREditorModule = IVREditorModule::Get();
+			VREditorModule.EnableVREditor(true);
+		}
+		else if (FParse::Param(FCommandLine::Get(), TEXT("ForceVREditor")))
+		{
+			GEngine->DeferredCommands.Add(TEXT("VREd.ForceVRMode"));
+		}
+	}
+}
+
 int32 EditorInit( IEngineLoop& EngineLoop )
 {
 	// Create debug exec.	
@@ -125,19 +141,8 @@ int32 EditorInit( IEngineLoop& EngineLoop )
 		}
 	}
 
-
 	// Go straight to VR mode if we were asked to
-	{
-		if( !bIsImmersive && FParse::Param( FCommandLine::Get(), TEXT( "VREditor" ) ) )
-		{
-			IVREditorModule& VREditorModule = IVREditorModule::Get();
-			VREditorModule.EnableVREditor( true );
-		}
-		else if( FParse::Param( FCommandLine::Get(), TEXT( "ForceVREditor" ) ) )
-		{
-			GEngine->DeferredCommands.Add( TEXT( "VREd.ForceVRMode" ) );
-		}
-	}
+	CheckAndMaybeGoToVRModeInternal(bIsImmersive);
 
 	// Check for automated build/submit option
 	const bool bDoAutomatedMapBuild = FParse::Param( FCommandLine::Get(), TEXT("AutomatedMapBuild") );
@@ -157,15 +162,37 @@ int32 EditorInit( IEngineLoop& EngineLoop )
 
 		if( FEngineAnalytics::IsAvailable() )
 		{
-			FEngineAnalytics::GetProvider().RecordEvent( 
-				TEXT( "Editor.Performance.Startup" ), 
+			FEngineAnalytics::GetProvider().RecordEvent(
+				TEXT( "Editor.Performance.Startup" ),
 				TEXT( "Duration" ), FString::Printf( TEXT( "%.3f" ), StartupTime ) );
 		}
 	}
 
 	FModuleManager::LoadModuleChecked<IModuleInterface>(TEXT("HierarchicalLODOutliner"));
 
-	// this will be ultimately returned from main(), so no error should be 0.
+	// This will be ultimately returned from main(), so no error should be 0.
+	return 0;
+}
+
+int32 EditorReinit()
+{
+	// Are we in immersive mode?
+	const bool bIsImmersive = FPaths::IsProjectFilePathSet() && FParse::Param(FCommandLine::Get(), TEXT("immersive"));
+	// Do final set up on the editor frame and show it
+	{
+		const bool bStartImmersive = bIsImmersive;
+		const bool bStartPIE = bIsImmersive;
+
+		// Tear down rendering thread once instead of doing it for every window being resized.
+		SCOPED_SUSPEND_RENDERING_THREAD(true);
+
+		// Startup Slate main frame and other editor windows
+		IMainFrameModule& MainFrameModule = FModuleManager::LoadModuleChecked<IMainFrameModule>(TEXT("MainFrame"));
+		MainFrameModule.RecreateDefaultMainFrame(bStartImmersive, bStartPIE);
+	}
+	// Go straight to VR mode if we were asked to
+	CheckAndMaybeGoToVRModeInternal(bIsImmersive);
+	// No error should be 0
 	return 0;
 }
 
diff --git a/Engine/Source/Editor/UnrealEd/Public/UnrealEdGlobals.h b/Engine/Source/Editor/UnrealEd/Public/UnrealEdGlobals.h
--- a/Engine/Source/Editor/UnrealEd/Public/UnrealEdGlobals.h
+++ b/Engine/Source/Editor/UnrealEd/Public/UnrealEdGlobals.h
@@ -7,4 +7,11 @@
 extern UNREALED_API class UUnrealEdEngine* GUnrealEd;
 
 UNREALED_API int32 EditorInit( class IEngineLoop& EngineLoop );
+
+/**
+ * Similar to EditorInit(), but the IMainFrameModule is recreated rather than just created.
+ * This function is useful e.g., to load a new layout without having to restart the whole editor.
+ */
+UNREALED_API int32 EditorReinit();
+
 UNREALED_API void EditorExit();
diff --git a/Engine/Source/Runtime/Core/Private/Misc/ConfigCacheIni.cpp b/Engine/Source/Runtime/Core/Private/Misc/ConfigCacheIni.cpp
--- a/Engine/Source/Runtime/Core/Private/Misc/ConfigCacheIni.cpp
+++ b/Engine/Source/Runtime/Core/Private/Misc/ConfigCacheIni.cpp
@@ -3452,10 +3452,10 @@ void FConfigCacheIni::InitializeConfigSystem()
 	FConfigCacheIni::LoadGlobalIniFile(GEditorPerProjectIni, TEXT("EditorPerProjectUserSettings"));
 
 	// Project agnostic editor ini files
-	static const FString EditorSettingsDir = FPaths::Combine(*FPaths::GameAgnosticSavedDir(), TEXT("Config")) + TEXT("/");
+	const FString EditorSettingsDir = FPaths::EngineEditorSettingsDir();
 	FConfigCacheIni::LoadGlobalIniFile(GEditorSettingsIni, TEXT("EditorSettings"), nullptr, false, false, true, *EditorSettingsDir);
-	FConfigCacheIni::LoadGlobalIniFile(GEditorLayoutIni, TEXT("EditorLayout"), nullptr, false, false, true, *EditorSettingsDir);
 	FConfigCacheIni::LoadGlobalIniFile(GEditorKeyBindingsIni, TEXT("EditorKeyBindings"), nullptr, false, false, true, *EditorSettingsDir);
+	FConfigCacheIni::LoadGlobalIniFile(GEditorLayoutIni, TEXT("EditorLayout"), nullptr, false, false, true, *EditorSettingsDir);
 
 #endif
 #if PLATFORM_DESKTOP
diff --git a/Engine/Source/Runtime/Core/Private/Misc/Paths.cpp b/Engine/Source/Runtime/Core/Private/Misc/Paths.cpp
--- a/Engine/Source/Runtime/Core/Private/Misc/Paths.cpp
+++ b/Engine/Source/Runtime/Core/Private/Misc/Paths.cpp
@@ -202,6 +202,11 @@ FString FPaths::EngineConfigDir()
 	return FPaths::EngineDir() + TEXT("Config/");
 }
 
+FString FPaths::EngineEditorSettingsDir()
+{
+	return FPaths::GameAgnosticSavedDir() + TEXT("Config/");
+}
+
 FString FPaths::EngineIntermediateDir()
 {
 	return FPaths::EngineDir() + TEXT("Intermediate/");
@@ -217,6 +222,16 @@ FString FPaths::EnginePluginsDir()
 	return FPaths::EngineDir() + TEXT("Plugins/");
 }
 
+FString FPaths::EngineDefaultLayoutDir()
+{
+	return FPaths::EngineConfigDir() + TEXT("Layouts/");
+}
+
+FString FPaths::EngineUserLayoutDir()
+{
+	return FPaths::EngineEditorSettingsDir() + TEXT("Layouts/");
+}
+
 FString FPaths::EnterpriseDir()
 {
 	return FPaths::RootDir() + TEXT("Enterprise/");
diff --git a/Engine/Source/Runtime/Core/Public/Misc/Paths.h b/Engine/Source/Runtime/Core/Public/Misc/Paths.h
--- a/Engine/Source/Runtime/Core/Public/Misc/Paths.h
+++ b/Engine/Source/Runtime/Core/Public/Misc/Paths.h
@@ -61,6 +61,13 @@ public:
 	 */
 	static FString EngineConfigDir();
 
+	/**
+	 * Returns the Editor Settings directory of the engine
+	 *
+	 * @return Editor Settings directory.
+	 */
+	static FString EngineEditorSettingsDir();
+
 	/**
 	 * Returns the intermediate directory of the engine
 	 *
@@ -82,6 +89,20 @@ public:
 	 */
 	static FString EnginePluginsDir();
 
+	/**
+	 * Returns the directory for default Editor UI Layout files of the engine
+	 *
+	 * @return Directory for default Editor UI Layout files.
+	 */
+	static FString EngineDefaultLayoutDir();
+
+	/**
+	 * Returns the directory for user-generated Editor UI Layout files of the engine
+	 *
+	 * @return Directory for user-generated Editor UI Layout files.
+	 */
+	static FString EngineUserLayoutDir();
+
 	/** 
 	* Returns the base directory enterprise directory.
 	*
diff --git a/Engine/Source/Runtime/Slate/Private/Framework/Docking/LayoutService.cpp b/Engine/Source/Runtime/Slate/Private/Framework/Docking/LayoutService.cpp
--- a/Engine/Source/Runtime/Slate/Private/Framework/Docking/LayoutService.cpp
+++ b/Engine/Source/Runtime/Slate/Private/Framework/Docking/LayoutService.cpp
@@ -14,7 +14,7 @@ const FString& FLayoutSaveRestore::GetAdditionalLayoutConfigIni()
 }
 
 
-void FLayoutSaveRestore::SaveToConfig( FString ConfigFileName, const TSharedRef<FTabManager::FLayout>& LayoutToSave )
+void FLayoutSaveRestore::SaveToConfig( const FString& ConfigFileName, const TSharedRef<FTabManager::FLayout>& LayoutToSave )
 {
 	const FString LayoutAsString = FLayoutSaveRestore::PrepareLayoutStringForIni( LayoutToSave->ToString() );
 
@@ -39,6 +39,19 @@ TSharedRef<FTabManager::FLayout> FLayoutSaveRestore::LoadFromConfig( const FStri
 }
 
 
+void FLayoutSaveRestore::SaveSectionToConfig(const FString& InConfigFileName, const FString& InSectionName, const FText& InSectionValue)
+{
+	GConfig->SetText(EditorLayoutsSectionName, *InSectionName, InSectionValue, InConfigFileName);
+}
+
+FText FLayoutSaveRestore::LoadSectionFromConfig(const FString& InConfigFileName, const FString& InSectionName)
+{
+	FText LayoutString;
+	GConfig->GetText(EditorLayoutsSectionName, *InSectionName, LayoutString, InConfigFileName);
+	return LayoutString;
+}
+
+
 void FLayoutSaveRestore::MigrateConfig( const FString& OldConfigFileName, const FString& NewConfigFileName )
 {
 	TArray<FString> OldSectionStrings;
@@ -72,6 +85,12 @@ void FLayoutSaveRestore::MigrateConfig( const FString& OldConfigFileName, const
 }
 
 
+bool FLayoutSaveRestore::IsValidConfig(const FString& InConfigFileName)
+{
+	return GConfig->DoesSectionExist(EditorLayoutsSectionName, *InConfigFileName);
+}
+
+
 FString FLayoutSaveRestore::PrepareLayoutStringForIni(const FString& LayoutString)
 {
 	// Have to store braces as parentheses due to braces causing ini issues
diff --git a/Engine/Source/Runtime/Slate/Private/Framework/Docking/TabManager.cpp b/Engine/Source/Runtime/Slate/Private/Framework/Docking/TabManager.cpp
--- a/Engine/Source/Runtime/Slate/Private/Framework/Docking/TabManager.cpp
+++ b/Engine/Source/Runtime/Slate/Private/Framework/Docking/TabManager.cpp
@@ -802,8 +802,8 @@ TSharedRef<FTabManager::FLayout> FTabManager::PersistLayout() const
 
 void FTabManager::SavePersistentLayout()
 {
-	const TSharedRef<FLayout> MyLayout = this->PersistLayout();
-	OnPersistLayout_Handler.ExecuteIfBound( MyLayout );
+	const TSharedRef<FLayout> LayoutState = this->PersistLayout();
+	OnPersistLayout_Handler.ExecuteIfBound( LayoutState );
 }
 
 FTabSpawnerEntry& FTabManager::RegisterTabSpawner(const FName TabId, const FOnSpawnTab& OnSpawnTab, const FCanSpawnTab& CanSpawnTab)
@@ -1713,11 +1713,8 @@ void FTabManager::OnTabManagerClosing()
 {
 	CleanupPointerArray(DockAreas);
 
-	{
-		// Gather the persistent layout and allow a custom handler to persist it
-		const TSharedRef<FLayout> LayoutState = this->PersistLayout();
-		OnPersistLayout_Handler.ExecuteIfBound( LayoutState );
-	}
+	// Gather the persistent layout and allow a custom handler to persist it
+	SavePersistentLayout();
 
 	for (int32 DockAreaIndex=0; DockAreaIndex < DockAreas.Num(); ++DockAreaIndex)
 	{
@@ -1935,9 +1932,14 @@ void FGlobalTabmanager::SetActiveTab( const TSharedPtr<class SDockTab>& NewActiv
 
 FTabSpawnerEntry& FGlobalTabmanager::RegisterNomadTabSpawner(const FName TabId, const FOnSpawnTab& OnSpawnTab, const FCanSpawnTab& CanSpawnTab)
 {
-	ensure(!NomadTabSpawner->Contains(TabId));
+	// Sanity check
 	ensure(!IsLegacyTabType(TabId));
-
+	// Remove TabId if it was previously loaded. This allows re-loading the Editor UI layout without restarting the whole Editor (Window->Load Layout)
+	if (NomadTabSpawner->Contains(TabId))
+	{
+		UnregisterNomadTabSpawner(TabId);
+	}
+	// (Re)create and return NewSpawnerEntry
 	TSharedRef<FTabSpawnerEntry> NewSpawnerEntry = MakeShareable(new FTabSpawnerEntry(TabId, OnSpawnTab, CanSpawnTab));
 	NomadTabSpawner->Add(TabId, NewSpawnerEntry);
 	return NewSpawnerEntry.Get();
diff --git a/Engine/Source/Runtime/Slate/Public/Framework/Docking/LayoutService.h b/Engine/Source/Runtime/Slate/Public/Framework/Docking/LayoutService.h
--- a/Engine/Source/Runtime/Slate/Public/Framework/Docking/LayoutService.h
+++ b/Engine/Source/Runtime/Slate/Public/Framework/Docking/LayoutService.h
@@ -10,22 +10,47 @@ struct SLATE_API FLayoutSaveRestore
 	/** Gets the ini section label for the additional layout configs */
 	static const FString& GetAdditionalLayoutConfigIni();
 
-	/** Write the layout out into a named config file.
+	/**
+	 * Write the layout out into a named config file.
 	 *
-	 * @param ConfigFileName file to be save to
+	 * @param ConfigFileName file to be saved to.
 	 * @param LayoutToSave the layout to save.
 	 */
-	static void SaveToConfig( FString ConfigFileName, const TSharedRef<FTabManager::FLayout>& LayoutToSave );
+	static void SaveToConfig(const FString& ConfigFileName, const TSharedRef<FTabManager::FLayout>& LayoutToSave );
 
-	/** Given a named DefualtLayout, return any saved version of it from the given ini file, otherwise return the default, also default to open tabs based on bool
+	/**
+	 * Given a named DefaultLayout, return any saved version of it from the given ini file, otherwise return the default, also default to open tabs based on bool.
 	 *
-	 * @param ConfigFileName file to be used to load an existing layout
+	 * @param ConfigFileName file to be used to load an existing layout.
 	 * @param DefaultLayout the layout to be used if the file does not exist.
 	 *
-	 * @return Loaded layout or the default
+	 * @return Loaded layout or the default.
 	 */
 	static TSharedRef<FTabManager::FLayout> LoadFromConfig( const FString& ConfigFileName, const TSharedRef<FTabManager::FLayout>& DefaultLayout );
 
+	/**
+	 * Write the desired FText value into the desired section of a named config file.
+	 * This function should only be used to save localization names (e.g., LayoutName, LayoutDescription).
+	 * For saving the FTabManager::FLayout, check SaveToConfig.
+	 *
+	 * @param InConfigFileName file to be saved to.
+	 * @param InSectionName the section name where to save the value.
+	 * @param InSectionValue the value to save.
+	 */
+	static void SaveSectionToConfig(const FString& InConfigFileName, const FString& InSectionName, const FText& InSectionValue);
+
+	/**
+	 * Read the desired FText value from the desired section of a named config file.
+	 * This function should only be used to load localization names (e.g., LayoutName, LayoutDescription).
+	 * For loading the FTabManager::FLayout, check SaveToConfig.
+	 *
+	 * @param InConfigFileName file to be used to load an existing layout.
+	 * @param InSectionName the name of the section to be read.
+	 *
+	 * @return Loaded FText associated for that section.
+	 */
+	static FText LoadSectionFromConfig(const FString& InConfigFileName, const FString& InSectionName);
+
 	/**
 	 * Migrates the layout configuration from one config file to another.
 	 *
@@ -34,6 +59,13 @@ struct SLATE_API FLayoutSaveRestore
 	 */
 	static void MigrateConfig( const FString& OldConfigFileName, const FString& NewConfigFileName );
 
+	/**
+	 * It checks whether a file is a valid layout config file.
+	 * @param InConfigFileName file to be read.
+	 * @return Whether the file is a valid layout config file.
+	 */
+	static bool IsValidConfig(const FString& InConfigFileName);
+
 private:
 	/**
 	 * Make a Json string friendly for writing out to UE .ini config files.
diff --git a/Engine/Source/Runtime/Slate/Public/Framework/Docking/TabManager.h b/Engine/Source/Runtime/Slate/Public/Framework/Docking/TabManager.h
--- a/Engine/Source/Runtime/Slate/Public/Framework/Docking/TabManager.h
+++ b/Engine/Source/Runtime/Slate/Public/Framework/Docking/TabManager.h
@@ -55,9 +55,16 @@ struct FTabId
 
 	FString ToString() const
 	{
-		return (InstanceId == INDEX_NONE)
-			? TabType.ToString()
-			: FString::Printf( TEXT("%s : %d"), *(TabType.ToString()), InstanceId );
+		// This function is useful to allow saving the layout on disk.
+		// Alternative: We could save (InstanceId == INDEX_NONE) ? TabType.ToString() : FString::Printf( TEXT("%s : %d"), *(TabType.ToString()), InstanceId ), which would include the InstanceId.
+		// Problem: InstanceId depends on the Editor/Slate runtime session rather than the layout itself. I.e., this would lead to the exact same layout being saved differently, depending
+		// on how many tabs were opened/closed before arriving to that final layout, or in the order in which those tabs were created.
+		// Conclusion: We do not want to save the InstanceId, which is session-dependent rather than layout-dependent.
+		// More detailed explanation: When the Editor is opened, InstanceId is a static value starting in 0 and increasing every time a new FTabID is created (but not decreased when one is
+		// closed). Loading/saving layouts do not reset this number either (and there is no point in doing so given that it is a runtime variable). E.g., opening and closing the same tab
+		// multiple times will lead to higher InstanceId numbers in the final layout. In addition, creating tab A and then B would lead to A::InstanceId = 0, B::InstanceId = 1. While creating
+		// B first and A latter would lead to the opposite values. Any of these examples would wrongly make 2 exact layouts look "different" if we saved the InstanceIds.
+		return TabType.ToString();
 	}
 
 	FText ToText() const
