Enhacement JIRA UETOOL-1721: Fixed incremental build error (missing include).

diff --git a/Engine/Source/Editor/MainFrame/Private/Menus/LayoutsMenu.h b/Engine/Source/Editor/MainFrame/Private/Menus/LayoutsMenu.h
--- a/Engine/Source/Editor/MainFrame/Private/Menus/LayoutsMenu.h
+++ b/Engine/Source/Editor/MainFrame/Private/Menus/LayoutsMenu.h
@@ -4,6 +4,8 @@
 
 #pragma once
 
+#include "CoreMinimal.h"
+
 class UToolMenu;
 
 /**
