Enhacement JIRA UETOOL-1721: Fixed incremental build error (missing another include).

diff --git a/Engine/Source/Editor/MainFrame/Private/Menus/LayoutsMenu.cpp b/Engine/Source/Editor/MainFrame/Private/Menus/LayoutsMenu.cpp
--- a/Engine/Source/Editor/MainFrame/Private/Menus/LayoutsMenu.cpp
+++ b/Engine/Source/Editor/MainFrame/Private/Menus/LayoutsMenu.cpp
@@ -11,7 +11,9 @@
 #include "Logging/MessageLog.h"
 #include "Templates/SharedPointer.h"
 // Runtime
+#include "CoreGlobals.h"
 #include "Framework/Commands/UICommandInfo.h"
+#include "Framework/Docking/LayoutService.h"
 #include "Framework/Notifications/NotificationManager.h"
 #include "Misc/FileHelper.h"
 #include "Widgets/Notifications/SNotificationList.h"
