Enhacement JIRA UETOOL-1721: Fixed incremental build error (missing even more includes).

diff --git a/Engine/Source/Editor/MainFrame/Private/Menus/LayoutsMenu.cpp b/Engine/Source/Editor/MainFrame/Private/Menus/LayoutsMenu.cpp
--- a/Engine/Source/Editor/MainFrame/Private/Menus/LayoutsMenu.cpp
+++ b/Engine/Source/Editor/MainFrame/Private/Menus/LayoutsMenu.cpp
@@ -15,6 +15,7 @@
 #include "Framework/Commands/UICommandInfo.h"
 #include "Framework/Docking/LayoutService.h"
 #include "Framework/Notifications/NotificationManager.h"
+#include "Misc/ConfigCacheIni.h"
 #include "Misc/FileHelper.h"
 #include "Widgets/Notifications/SNotificationList.h"
 // Developer
