Fix to JIRA UE-79515:
1) After loading any layout, now the changes work again after closing and re-starting the editor.
2) Attempting to SaveAs or Override the current layout with a different file name now clears the description and tooltip description to avoid confusion.
3) In addition, improved the layout tooltip to include the file path in Layouts with description.

diff --git a/Engine/Source/Editor/MainFrame/Private/Menus/LayoutsMenu.cpp b/Engine/Source/Editor/MainFrame/Private/Menus/LayoutsMenu.cpp
--- a/Engine/Source/Editor/MainFrame/Private/Menus/LayoutsMenu.cpp
+++ b/Engine/Source/Editor/MainFrame/Private/Menus/LayoutsMenu.cpp
@@ -130,7 +130,7 @@ bool TrySaveLayoutOrWarnInternal(const FString& InSourceFilePath, const FString&
 	// Copy: Replace main layout with desired one
 	const bool bShouldReplace = true;
 	const bool bCopyEvenIfReadOnly = true;
-	const bool bCopyAttributes = true;
+	const bool bCopyAttributes = false; // If true, it could e.g., copy the read-only flag of DefaultLayout.ini and make all the save/load stuff stop working
 	if (COPY_Fail == IFileManager::Get().Copy(*InTargetFilePath, *InSourceFilePath, bShouldReplace, bCopyEvenIfReadOnly, bCopyAttributes))
 	{
 		FMessageLog EditorErrors("EditorErrors");
@@ -141,14 +141,14 @@ bool TrySaveLayoutOrWarnInternal(const FString& InSourceFilePath, const FString&
 		if (!FPaths::FileExists(InSourceFilePath))
 		{
 			Arguments.Add(TEXT("FileName"), FText::FromString(FPaths::ConvertRelativePathToFull(InSourceFilePath)));
-			TextBody = FText::Format(LOCTEXT("UnsuccessfulSave_NoExist_Notification", "Unsuccessful {WhatIs}, the desired file does not exist:\n{FileName}"), Arguments);
+			TextBody = FText::Format(LOCTEXT("UnsuccessfulSave_NoExist_Notification", "Unsuccessful {WhatIs}, the desired file does not exist. File path:\n{FileName}"), Arguments);
 			EditorErrors.Warning(TextBody);
 		}
 		// Target is read-only
 		else if (IFileManager::Get().IsReadOnly(*InTargetFilePath))
 		{
 			Arguments.Add(TEXT("FileName"), FText::FromString(FPaths::ConvertRelativePathToFull(InTargetFilePath)));
-			TextBody = FText::Format(LOCTEXT("UnsuccessfulSave_ReadOnly_Notification", "Unsuccessful {WhatIs}, the target file path is read-only\n{FileName}"), Arguments);
+			TextBody = FText::Format(LOCTEXT("UnsuccessfulSave_ReadOnly_Notification", "Unsuccessful {WhatIs}, the target file path is read-only. File path:\n{FileName}"), Arguments);
 			EditorErrors.Warning(TextBody);
 		}
 		// Target and source are the same
@@ -211,9 +211,18 @@ FText GetDisplayTextInternal(const FString& InString)
 	const FText DisplayName = FText::FromString(FName::NameToDisplayString(DisplayNameText.ToString(), bIsBool));
 	return DisplayName;
 }
-FText GetTooltipTextInternal(const FString& InKindOfFile, const FText& InDisplayName, const FString& InLayoutFilePath)
+FText GetTooltipTextInternal(const FText& InDisplayName, const FString& InLayoutFilePath, const FText& InLayoutName, const int32 LayoutIndex)
 {
-	return FText::FromString(InKindOfFile + FString(" name:\n") + InDisplayName.ToString() + FString("\n\nFull file path:\n") + InLayoutFilePath);
+	FText Tooltip;
+	if (InLayoutName.IsEmpty())
+	{
+		Tooltip = FText::Format(LOCTEXT("DisplayName{0}", "Layout name:\n{1}\n\nFull file path:\n{2}"), FText::AsNumber(LayoutIndex), InDisplayName, FText::FromString(InLayoutFilePath));
+	}
+	else
+	{
+		Tooltip = FText::Format(LOCTEXT("DisplayName{0}", "Description:\n{1}.\n\nFull file path:\n{2}"), FText::AsNumber(LayoutIndex), InLayoutName, FText::FromString(InLayoutFilePath));
+	}
+	return Tooltip;
 }
 
 void DisplayLayoutsInternal(FToolMenuSection& InSection, const TArray<TSharedPtr<FUICommandInfo>>& InXLayoutCommands, const TArray<FString> InLayoutIniFileNames, const FString InLayoutsDirectory)
@@ -231,7 +240,7 @@ void DisplayLayoutsInternal(FToolMenuSection& InSection, const TArray<TSharedPtr
 			const FText LayoutDescription = FLayoutSaveRestore::LoadSectionFromConfig(LayoutFilePath, "LayoutDescription");
 			// If no localization name, then display the file name
 			const FText DisplayName = (!LayoutName.IsEmpty() ? LayoutName : GetDisplayTextInternal(InLayoutIniFileNames[LayoutIndex]));
-			const FText Tooltip = (!LayoutDescription.IsEmpty() ? LayoutDescription : GetTooltipTextInternal("Default layout", DisplayName, LayoutFilePath));
+			const FText Tooltip = GetTooltipTextInternal(DisplayName, LayoutFilePath, LayoutDescription, LayoutIndex);
 			InSection.AddMenuEntry(InXLayoutCommands[LayoutIndex], DisplayName, Tooltip);
 		}
 	}
@@ -364,7 +373,8 @@ void FLayoutsMenuLoad::LoadLayout(const FString& InLayoutPath)
 	const FString& SourceFilePath = InLayoutPath;
 	const FString& TargetFilePath = GEditorLayoutIni;
 	const bool bCleanLayoutNameAndDescriptionFieldsIfNoSameValues = false;
-	const bool SucessfullySaved = TrySaveLayoutOrWarnInternal(SourceFilePath, TargetFilePath, LOCTEXT("LoadLayoutText", "layout load"), bCleanLayoutNameAndDescriptionFieldsIfNoSameValues);
+	const bool bShouldAskBeforeCleaningLayoutNameAndDescriptionFields = false;
+	const bool SucessfullySaved = TrySaveLayoutOrWarnInternal(SourceFilePath, TargetFilePath, LOCTEXT("LoadLayoutText", "layout load"), bCleanLayoutNameAndDescriptionFieldsIfNoSameValues, bShouldAskBeforeCleaningLayoutNameAndDescriptionFields);
 	// Reload current layout
 	if (SucessfullySaved)
 	{
@@ -426,7 +436,8 @@ void FLayoutsMenuLoad::ImportLayout()
 					const FString& SourceFilePath = LayoutFilePath;
 					const FString& TargetFilePath = FPaths::Combine(FPaths::GetPath(UserLayoutsDirectory), FPaths::GetCleanFilename(LayoutFilePath));
 					const bool bCleanLayoutNameAndDescriptionFieldsIfNoSameValues = false;
-					TrySaveLayoutOrWarnInternal(SourceFilePath, TargetFilePath, TrySaveLayoutOrWarnInternalText, bCleanLayoutNameAndDescriptionFieldsIfNoSameValues);
+					const bool bShouldAskBeforeCleaningLayoutNameAndDescriptionFields = false;
+					TrySaveLayoutOrWarnInternal(SourceFilePath, TargetFilePath, TrySaveLayoutOrWarnInternalText, bCleanLayoutNameAndDescriptionFieldsIfNoSameValues, bShouldAskBeforeCleaningLayoutNameAndDescriptionFields);
 				}
 				// File is not a layout file, warn the user
 				else
@@ -453,7 +464,8 @@ void FLayoutsMenuLoad::ImportLayout()
 				const FString& SourceFilePath = FirstGoodLayoutFile;
 				const FString& TargetFilePath = GEditorLayoutIni;
 				const bool bCleanLayoutNameAndDescriptionFieldsIfNoSameValues = false;
-				const bool SucessfullySaved = TrySaveLayoutOrWarnInternal(SourceFilePath, TargetFilePath, TrySaveLayoutOrWarnInternalText, bCleanLayoutNameAndDescriptionFieldsIfNoSameValues);
+				const bool bShouldAskBeforeCleaningLayoutNameAndDescriptionFields = false;
+				const bool SucessfullySaved = TrySaveLayoutOrWarnInternal(SourceFilePath, TargetFilePath, TrySaveLayoutOrWarnInternalText, bCleanLayoutNameAndDescriptionFieldsIfNoSameValues, bShouldAskBeforeCleaningLayoutNameAndDescriptionFields);
 				// Reload current layout
 				if (SucessfullySaved)
 				{
@@ -591,7 +603,7 @@ void FLayoutsMenuSave::OverrideUserLayout(const int32 InLayoutIndex)
 	// Replace desired layout with current one
 	const bool bCleanLayoutNameAndDescriptionFieldsIfNoSameValues = true;
 	const bool bShouldAskBeforeCleaningLayoutNameAndDescriptionFields = false;
-	TrySaveLayoutOrWarnInternal(SourceFilePath, TargetFilePath, LOCTEXT("OverrideLayoutText", "layout override"), bCleanLayoutNameAndDescriptionFieldsIfNoSameValues);
+	TrySaveLayoutOrWarnInternal(SourceFilePath, TargetFilePath, LOCTEXT("OverrideLayoutText", "layout override"), bCleanLayoutNameAndDescriptionFieldsIfNoSameValues, bShouldAskBeforeCleaningLayoutNameAndDescriptionFields);
 }
 
 void FLayoutsMenuSave::SaveLayout()
