Enhancement: Automatically close unrecognized tabs when the Editor is opened (e.g., the tabs of a plugin that was disabled). However, they are kept saved in the layout file in case they are recognized in future starts (e.g., if the plugin is re-enabled).

diff --git a/Engine/Source/Editor/MainFrame/Private/Menus/LayoutsMenu.cpp b/Engine/Source/Editor/MainFrame/Private/Menus/LayoutsMenu.cpp
--- a/Engine/Source/Editor/MainFrame/Private/Menus/LayoutsMenu.cpp
+++ b/Engine/Source/Editor/MainFrame/Private/Menus/LayoutsMenu.cpp
@@ -246,6 +246,91 @@ void DisplayLayoutsInternal(FToolMenuSection& InSection, const TArray<TSharedPtr
 	}
 }
 
+// GetOriginalEditorLayoutIniFilePathInternal() and GetDuplicatedEditorLayoutIniFilePathInternal() are used because sometimes the layout saved is not the same than the one loaded, even
+// though the visual display and screenshot are 100% the same. In those cases, we still want to show the check mark in the load/save/remove menu indicating that the layout is the same
+// than the one in the loaded ini file, even though the actual files might not be exactly the same. In addition, this also helps when temporarily closing unrecognized tabs (i.e., the
+// ones that FTabManager::SpawnTab cannot recognized and keeps closed but still saves them in the layout).
+const FString& GetOriginalEditorLayoutIniFilePathInternal()
+{
+	static const FString OriginalEditorLayoutIniFilePath = GEditorLayoutIni + FString("_orig.ini");
+	return OriginalEditorLayoutIniFilePath;
+}
+
+const FString& GetDuplicatedEditorLayoutIniFilePathInternal()
+{
+	static const FString DuplicatedEditorLayoutIniFilePath = GEditorLayoutIni + FString("_temp.ini");
+	return DuplicatedEditorLayoutIniFilePath;
+}
+
+bool AreFilesIdenticalInternal(const FString& InFirstFileFullPath, const FString& InSecondFileFullPath)
+{
+	// Checked if same file. I.e.,
+		// 1. Same size
+		// 2. And same internal text
+	const bool bHaveSameSize = (IFileManager::Get().FileSize(*InFirstFileFullPath) == IFileManager::Get().FileSize(*InSecondFileFullPath));
+	// Same size --> Same layout file?
+	if (bHaveSameSize)
+	{
+		// Read files and check whether they have the exact same text
+		FString StringFirstFileFullPath;
+		FFileHelper::LoadFileToString(StringFirstFileFullPath, *InFirstFileFullPath);
+		FString StringSecondFileFullPath;
+		FFileHelper::LoadFileToString(StringSecondFileFullPath, *InSecondFileFullPath);
+		// (No) same text = (No) same layout file
+		return (StringFirstFileFullPath == StringSecondFileFullPath);
+	}
+	// No same size = No same layout file
+	else
+	{
+		return false;
+	}
+}
+
+void RemoveTempEditorLayoutIniFilesInternal()
+{
+	const bool bRequireExists = false;
+	const bool bEvenIfReadOnly = true;
+	const bool bIsQuiet = false;
+	// DuplicatedEditorLayoutIniFilePath
+	const FString& DuplicatedEditorLayoutIniFilePath = GetDuplicatedEditorLayoutIniFilePathInternal();
+	IFileManager::Get().Delete(*DuplicatedEditorLayoutIniFilePath, bRequireExists, bEvenIfReadOnly, bIsQuiet);
+	GConfig->UnloadFile(DuplicatedEditorLayoutIniFilePath);
+	// OriginalEditorLayoutIniFilePath
+	const FString& OriginalEditorLayoutIniFilePath = GetOriginalEditorLayoutIniFilePathInternal();
+	IFileManager::Get().Delete(*OriginalEditorLayoutIniFilePath, bRequireExists, bEvenIfReadOnly, bIsQuiet);
+	GConfig->UnloadFile(OriginalEditorLayoutIniFilePath);
+}
+
+bool IsLayoutCheckedInternal(const FString& InLayoutFullPath, const bool bCheckTempFileToo)
+{
+	// If same file, return true
+	if (AreFilesIdenticalInternal(InLayoutFullPath, GEditorLayoutIni))
+	{
+		return true;
+	}
+	// No same size, check if same than temporary one
+	else if (bCheckTempFileToo)
+	{
+		const FString& OriginalEditorLayoutIniFilePath = GetOriginalEditorLayoutIniFilePathInternal();
+		const FString& DuplicatedEditorLayoutIniFilePath = GetDuplicatedEditorLayoutIniFilePathInternal();
+		if (AreFilesIdenticalInternal(GEditorLayoutIni, DuplicatedEditorLayoutIniFilePath))
+		{
+			return AreFilesIdenticalInternal(InLayoutFullPath, OriginalEditorLayoutIniFilePath);
+		}
+		// If GEditorLayoutIni != DuplicatedEditorLayoutIniFilePath, remove DuplicatedEditorLayoutIniFilePath & OriginalEditorLayoutIniFilePath
+		else
+		{
+			RemoveTempEditorLayoutIniFilesInternal();
+			return false;
+		}
+	}
+	// No same size, and we should not check if same than temporary ones, so then it is false
+	else
+	{
+		return false;
+	}
+}
+
 void MakeXLayoutsMenuInternal(UToolMenu* InToolMenu, const TArray<TSharedPtr<FUICommandInfo>>& InXLayoutCommands, const TArray<TSharedPtr<FUICommandInfo>>& InXUserLayoutCommands, const bool bDisplayDefaultLayouts)
 {
 	// Update GEditorLayoutIni file. Otherwise, we could not track the changes the user did since the layout was loaded
@@ -280,7 +365,7 @@ void MakeXLayoutsMenuInternal(UToolMenu* InToolMenu, const TArray<TSharedPtr<FUI
  * @param	InLayoutIndex  Index from the selected layout.
  * @return true if the selected layout can be read.
  */
-bool CanChooseLayoutWhenRead(const int32 InLayoutIndex)
+bool CanChooseLayoutWhenReadInternal(const int32 InLayoutIndex)
 {
 	return true;
 }
@@ -290,7 +375,7 @@ bool CanChooseLayoutWhenRead(const int32 InLayoutIndex)
  * @param	InLayoutIndex  Index from the selected user-created layout.
  * @return true if the selected user-created layout can be read.
  */
-bool CanChooseUserLayoutWhenRead(const int32 InLayoutIndex)
+bool CanChooseUserLayoutWhenReadInternal(const int32 InLayoutIndex)
 {
 	return true;
 }
@@ -301,7 +386,7 @@ bool CanChooseUserLayoutWhenRead(const int32 InLayoutIndex)
  * @param	InLayoutIndex  Index from the selected layout.
  * @return true if the selected layout can be modified/removed.
  */
-bool CanChooseLayoutWhenWrite(const int32 InLayoutIndex)
+bool CanChooseLayoutWhenWriteInternal(const int32 InLayoutIndex)
 {
 	return false;
 }
@@ -311,11 +396,21 @@ bool CanChooseLayoutWhenWrite(const int32 InLayoutIndex)
  * @param	InLayoutIndex  Index from the selected user-created layout.
  * @return true if the selected user-created layout can be modified/removed.
  */
-bool CanChooseUserLayoutWhenWrite(const int32 InLayoutIndex)
+bool CanChooseUserLayoutWhenWriteInternal(const int32 InLayoutIndex)
 {
 	return true;
 }
 
+void SaveLayoutWithoutRemovingTempLayoutFiles()
+{
+	// Save the layout into the Editor
+	FGlobalTabmanager::Get()->SaveAllVisualState();
+	// Write the saved layout to disk (if it has changed since the last time it was read/written)
+	// We must set bRead = true. Otherwise, FLayoutsMenuLoad::ReloadCurrentLayout() would reload the old config file (because it would be cached on memory)
+	const bool bRead = true;
+	GConfig->Flush(bRead, GEditorLayoutIni);
+}
+
 
 
 void FLayoutsMenuLoad::MakeLoadLayoutsMenu(UToolMenu* InToolMenu)
@@ -348,7 +443,7 @@ bool FLayoutsMenuLoad::CanLoadChooseLayout(const int32 InLayoutIndex)
 	{
 		return false;
 	}
-	return !FLayoutsMenuBase::IsLayoutChecked(InLayoutIndex) && CanChooseLayoutWhenRead(InLayoutIndex);
+	return !FLayoutsMenuBase::IsLayoutChecked(InLayoutIndex) && CanChooseLayoutWhenReadInternal(InLayoutIndex);
 }
 bool FLayoutsMenuLoad::CanLoadChooseUserLayout(const int32 InLayoutIndex)
 {
@@ -356,15 +451,45 @@ bool FLayoutsMenuLoad::CanLoadChooseUserLayout(const int32 InLayoutIndex)
 	{
 		return false;
 	}
-	return !FLayoutsMenuBase::IsUserLayoutChecked(InLayoutIndex) && CanChooseUserLayoutWhenRead(InLayoutIndex);
+	return !FLayoutsMenuBase::IsUserLayoutChecked(InLayoutIndex) && CanChooseUserLayoutWhenReadInternal(InLayoutIndex);
 }
 
 void FLayoutsMenuLoad::ReloadCurrentLayout()
 {
+	// Create duplicated ini file (OriginalEditorLayoutIniFilePath)
+	// Explanation:
+	//     Assume a layout is saved with (at least) a window that is dependent on a plugin. If that plugin is disabled and the editor restarted, that window will be saved on the layout but will
+	//     not visually appear. We still wanna save the layout with it, so if its plugin is re-enabled, the window appear again. However, while the plugin is disabled, the layout ini file changes
+	//     to reflect that the plugin is not opened.
+	// Technical details:
+	//     Rather than changing the string generated in the ini file (which could affect other parts of the code), we will duplicate the ini file when loaded. If the ini file is different than
+	//     its duplicated copy, then some widget is missing (most probably due to disabled plugins). If that is the case, we will re-save the ini file without telling UE that it changed. This way,
+	//     the ini file would match its original one, and it would only be re-modified if the user modifies the layout (but in that case it should no longer match the original one).
+	const bool bShouldReplace = true;
+	const bool bEvenIfReadOnly = true;
+	const bool bCopyAttributes = false; // If true, it could e.g., copy the read-only flag of DefaultLayout.ini and make all the save/load stuff stop working
+	const FString& OriginalEditorLayoutIniFilePath = GetOriginalEditorLayoutIniFilePathInternal();
+	IFileManager::Get().Copy(*OriginalEditorLayoutIniFilePath, *GEditorLayoutIni, bShouldReplace, bEvenIfReadOnly, bCopyAttributes);
+	GConfig->UnloadFile(OriginalEditorLayoutIniFilePath);
 	// Editor is reset on-the-fly
 	FUnrealEdMisc::Get().AllowSavingLayoutOnClose(false);
 	EditorReinit();
 	FUnrealEdMisc::Get().AllowSavingLayoutOnClose(true);
+	// Save layout and create duplicated ini file (DuplicatedEditorLayoutIniFilePath)
+	SaveLayoutWithoutRemovingTempLayoutFiles();
+	// If same file, remove temp files
+	const bool bCheckTempFileToo = false;
+	if (IsLayoutCheckedInternal(OriginalEditorLayoutIniFilePath, bCheckTempFileToo))
+	{
+		RemoveTempEditorLayoutIniFilesInternal();
+	}
+	// Else, create DuplicatedEditorLayoutIniFilePath
+	else
+	{
+		const FString& DuplicatedEditorLayoutIniFilePath = GetDuplicatedEditorLayoutIniFilePathInternal();
+		IFileManager::Get().Copy(*DuplicatedEditorLayoutIniFilePath, *GEditorLayoutIni, bShouldReplace, bEvenIfReadOnly, bCopyAttributes);
+		GConfig->UnloadFile(DuplicatedEditorLayoutIniFilePath);
+	}
 }
 
 void FLayoutsMenuLoad::LoadLayout(const FString& InLayoutPath)
@@ -567,11 +692,11 @@ void FLayoutsMenuSave::MakeSaveLayoutsMenu(UToolMenu* InToolMenu)
 
 bool FLayoutsMenuSave::CanSaveChooseLayout(const int32 InLayoutIndex)
 {
-	return !FLayoutsMenuBase::IsLayoutChecked(InLayoutIndex) && CanChooseLayoutWhenWrite(InLayoutIndex);
+	return !FLayoutsMenuBase::IsLayoutChecked(InLayoutIndex) && CanChooseLayoutWhenWriteInternal(InLayoutIndex);
 }
 bool FLayoutsMenuSave::CanSaveChooseUserLayout(const int32 InLayoutIndex)
 {
-	return !FLayoutsMenuBase::IsUserLayoutChecked(InLayoutIndex) && CanChooseUserLayoutWhenWrite(InLayoutIndex);
+	return !FLayoutsMenuBase::IsUserLayoutChecked(InLayoutIndex) && CanChooseUserLayoutWhenWriteInternal(InLayoutIndex);
 }
 
 void FLayoutsMenuSave::OverrideLayout(const int32 InLayoutIndex)
@@ -608,12 +733,15 @@ void FLayoutsMenuSave::OverrideUserLayout(const int32 InLayoutIndex)
 
 void FLayoutsMenuSave::SaveLayout()
 {
-	// Save the layout into the Editor
-	FGlobalTabmanager::Get()->SaveAllVisualState();
-	// Write the saved layout to disk (if it has changed since the last time it was read/written)
-	// We must set bRead = true. Otherwise, FLayoutsMenuLoad::ReloadCurrentLayout() would reload the old config file (because it would be cached on memory)
-	const bool bRead = true;
-	GConfig->Flush(bRead, GEditorLayoutIni);
+	// Save layout
+	SaveLayoutWithoutRemovingTempLayoutFiles();
+	// Remove temporary Editor Layout ini files if the layout (thus also GEditorLayoutIni) changed
+	const bool bCheckTempFileToo = false;
+	const FString& DuplicatedEditorLayoutIniFilePath = GetDuplicatedEditorLayoutIniFilePathInternal();
+	if (!IsLayoutCheckedInternal(DuplicatedEditorLayoutIniFilePath, bCheckTempFileToo))
+	{
+		RemoveTempEditorLayoutIniFilesInternal();
+	}
 }
 
 void FLayoutsMenuSave::SaveLayoutAs()
@@ -666,11 +794,11 @@ void FLayoutsMenuRemove::MakeRemoveLayoutsMenu(UToolMenu* InToolMenu)
 
 bool FLayoutsMenuRemove::CanRemoveChooseLayout(const int32 InLayoutIndex)
 {
-	return CanChooseLayoutWhenWrite(InLayoutIndex);
+	return CanChooseLayoutWhenWriteInternal(InLayoutIndex);
 }
 bool FLayoutsMenuRemove::CanRemoveChooseUserLayout(const int32 InLayoutIndex)
 {
-	return CanChooseUserLayoutWhenWrite(InLayoutIndex);
+	return CanChooseUserLayoutWhenWriteInternal(InLayoutIndex);
 }
 
 void FLayoutsMenuRemove::RemoveLayout(const int32 InLayoutIndex)
@@ -760,30 +888,6 @@ void FLayoutsMenuRemove::RemoveUserLayouts()
 
 
 
-bool IsLayoutCheckedInternal(const FString& InLayoutFullPath)
-{
-	// Checked if same file. I.e.,
-		// 1. Same size
-		// 2. And same internal text
-	const bool bHaveSameSize = (IFileManager::Get().FileSize(*GEditorLayoutIni) == IFileManager::Get().FileSize(*InLayoutFullPath));
-	// Same size --> Same layout file?
-	if (bHaveSameSize)
-	{
-		// Read files and check whether they have the exact same text
-		FString StringGEditorLayoutIni;
-		FFileHelper::LoadFileToString(StringGEditorLayoutIni, *GEditorLayoutIni);
-		FString StringLayoutFullPath;
-		FFileHelper::LoadFileToString(StringLayoutFullPath, *InLayoutFullPath);
-		// (No) same text = (No) same layout file
-		return (StringGEditorLayoutIni == StringLayoutFullPath);
-	}
-	// No same size = No same layout file
-	else
-	{
-		return false;
-	}
-}
-
 FString FLayoutsMenuBase::GetLayout(const int32 InLayoutIndex)
 {
 	// Get LayoutsDirectory path, layout init files, and desired layout path
@@ -815,13 +919,15 @@ bool FLayoutsMenuBase::IsThereUserLayouts()
 bool FLayoutsMenuBase::IsLayoutChecked(const int32 InLayoutIndex)
 {
 	// Check if the desired layout file matches the one currently loaded
-	return IsLayoutCheckedInternal(GetLayout(InLayoutIndex));
+	const bool bCheckTempFileToo = true;
+	return IsLayoutCheckedInternal(GetLayout(InLayoutIndex), bCheckTempFileToo);
 }
 
 bool FLayoutsMenuBase::IsUserLayoutChecked(const int32 InLayoutIndex)
 {
 	// Check if the desired layout file matches the one currently loaded
-	return IsLayoutCheckedInternal(GetUserLayout(InLayoutIndex));
+	const bool bCheckTempFileToo = true;
+	return IsLayoutCheckedInternal(GetUserLayout(InLayoutIndex), bCheckTempFileToo);
 }
 
 #undef LOCTEXT_NAMESPACE
diff --git a/Engine/Source/Runtime/Slate/Private/Framework/Docking/TabManager.cpp b/Engine/Source/Runtime/Slate/Private/Framework/Docking/TabManager.cpp
--- a/Engine/Source/Runtime/Slate/Private/Framework/Docking/TabManager.cpp
+++ b/Engine/Source/Runtime/Slate/Private/Framework/Docking/TabManager.cpp
@@ -1498,25 +1498,16 @@ TSharedPtr<SDockTab> FTabManager::SpawnTab( const FTabId& TabId, const TSharedPt
 	}
 
 	// The tab was allowed to be spawned but failed for some reason
-	if (bSpawningAllowedBySpawner && !NewTabWidget.IsValid() )
-	{
-		// We don't know how to spawn this tab.
-		// Make a dummy tab so that things aren't entirely broken.
-
-		NewTabWidget = SNew(SDockTab)
-			.Label( TabId.ToText() )
-			.ShouldAutosize( false )
-			[
-				SNew(SBox)
-				.HAlign(HAlign_Center)
-				.VAlign(VAlign_Center)
-				[
-					SNew(STextBlock)
-					.Text( NSLOCTEXT("TabManagement", "Unrecognized", "unrecognized tab") )
-				]
-			];
-
-		NewTabWidget->SetLayoutIdentifier(TabId);
+	if (bSpawningAllowedBySpawner && !NewTabWidget.IsValid())
+	{
+		// We don't know how to spawn this tab. 2 alternatives:
+		// 1) Make a dummy tab so that things aren't entirely broken (previous versions did this).
+		// 2) Do not open the widget, but keep it saved in the layout, so this function should still return it.
+		const FString StringToDisplay = (Spawner.IsValid() && !Spawner->GetDisplayName().IsEmpty() ? Spawner->GetDisplayName().ToString() : TabId.TabType.ToString());
+		UE_LOG(LogSlate, Warning,
+			TEXT("The tab \"%s\" attempted to spawn but failed for some reason. It will not be displayed but it will be saved in the layout settings file."),
+			*(!StringToDisplay.IsEmpty() ? StringToDisplay : FString("Unknown"))
+			);
 	}
 
 	if (NewTabWidget.IsValid())
