Fixed some typos (grammatical errors) in a dialog from LayoutsMenu.cpp.

diff --git a/Engine/Source/Editor/MainFrame/Private/Menus/LayoutsMenu.cpp b/Engine/Source/Editor/MainFrame/Private/Menus/LayoutsMenu.cpp
--- a/Engine/Source/Editor/MainFrame/Private/Menus/LayoutsMenu.cpp
+++ b/Engine/Source/Editor/MainFrame/Private/Menus/LayoutsMenu.cpp
@@ -845,6 +845,23 @@ int32 GetNumberLayoutFiles(const FString& InLayoutsDirectory)
 	return NumberLayoutFiles;
 }
 
+FText GetTextBodyForRemoveLayoutDialog(const int32 NumberUserLayoutFiles)
+{
+	// 1 layout file found
+	FText TextToAddDependingOnNumberFiles;
+	if (NumberUserLayoutFiles == 1)
+	{
+		TextToAddDependingOnNumberFiles = LOCTEXT("TextToAddDependingOnNumberFiles1", "");
+	}
+	// Any other number
+	else
+	{
+		const FText TextNumberUserLayoutFiles = FText::FromString(FString::FromInt(NumberUserLayoutFiles));
+		TextToAddDependingOnNumberFiles = FText::Format(LOCTEXT("TextToAddDependingOnNumberFiles", " ({0} profiles were found)"), TextNumberUserLayoutFiles);
+	}
+	return FText::Format(LOCTEXT("ActionRemoveAllUserLayoutMsg", "Are you sure you want to permanently remove all the layout profiles created by the user{0}? This action cannot be undone."), TextToAddDependingOnNumberFiles);
+}
+
 void FLayoutsMenuRemove::RemoveUserLayouts()
 {
 	// (Create if it does not exist and) Get UserLayoutsDirectory path
@@ -855,8 +872,7 @@ void FLayoutsMenuRemove::RemoveUserLayouts()
 	if (NumberUserLayoutFiles > 0)
 	{
 		// Are you sure you want to do this?
-		const FText TextNumberUserLayoutFiles  = FText::FromString(FString::FromInt(NumberUserLayoutFiles));
-		const FText TextBody = FText::Format(LOCTEXT("ActionRemoveAllUserLayoutMsg", "Are you sure you want to permanently remove all the {0} layout profiles created by the user? This action cannot be undone."), TextNumberUserLayoutFiles);
+		const FText TextBody = GetTextBodyForRemoveLayoutDialog(NumberUserLayoutFiles);
 		const FText TextTitle = LOCTEXT("RemoveAllUserLayouts_Title", "Remove All User-Created Layouts");
 		if (EAppReturnType::Ok != OpenMsgDlgInt(EAppMsgType::OkCancel, TextBody, TextTitle))
 		{
