Bug fix: Close Asset Editors warning while loading layout had 2 identical options: No and Cancel. Removed the Cancel option to avoid the duplicated option.

diff --git a/Engine/Source/Editor/MainFrame/Private/Menus/LayoutsMenu.cpp b/Engine/Source/Editor/MainFrame/Private/Menus/LayoutsMenu.cpp
--- a/Engine/Source/Editor/MainFrame/Private/Menus/LayoutsMenu.cpp
+++ b/Engine/Source/Editor/MainFrame/Private/Menus/LayoutsMenu.cpp
@@ -522,9 +522,9 @@ bool CheckAskUserAndClosePIESIEAndAssetEditors(const FText& InitialMessage)
 		IfYesText = LOCTEXT("CheckAskUserAndClosePIESIEAndAssetEditorsIfYesBodyEditorAssets", "If \"Yes\", all open Asset Editors will be closed. Any unsaved changes in those will also be lost.");
 	}
 	// FMessageDialog
-	const FText IfNoText = LOCTEXT("CheckAskUserAndClosePIESIEAndAssetEditorsIfNoBody", "If \"No\" or \"Cancel\", you can manually reload the layout from the \"User Layouts\" section later.");
+	const FText IfNoText = LOCTEXT("CheckAskUserAndClosePIESIEAndAssetEditorsIfNoBody", "If \"No\", you can manually reload the layout from the \"User Layouts\" section later.");
 	const FText TextBody = FText::Format(LOCTEXT("ClosePIESIEAssetEditorsBody", "{0}\n\n{1}{2}\n\n{3}"), InitialMessage, IfYesText, OpenedEditorAssets, IfNoText);
-	if (EAppReturnType::Yes != FMessageDialog::Open(EAppMsgType::YesNoCancel, TextBody, &TextTitle))
+	if (EAppReturnType::Yes != FMessageDialog::Open(EAppMsgType::YesNo, TextBody, &TextTitle))
 	{
 		return false;
 	}
