Added comments to help update the default values of the Editor layout (i.e., if changes are applied to "LevelEditor_Layout_v1.1" or "UnrealEd_Layout_v1.4").

diff --git a/Engine/Source/Editor/LevelEditor/Private/SLevelEditor.cpp b/Engine/Source/Editor/LevelEditor/Private/SLevelEditor.cpp
--- a/Engine/Source/Editor/LevelEditor/Private/SLevelEditor.cpp
+++ b/Engine/Source/Editor/LevelEditor/Private/SLevelEditor.cpp
@@ -1230,6 +1230,20 @@ TSharedRef<SWidget> SLevelEditor::RestoreContentArea( const TSharedRef<SDockTab>
 	// or there wont be any tab spawners for the modes.
 	RefreshEditorModeCommands();
 
+	// IMPORTANT: If you want to change the default value of "LevelEditor_Layout_v1.1" or "UnrealEd_Layout_v1.4" (even if you only change their version numbers), these are the steps to follow:
+	// 1. Check out Engine\Config\Layouts\DefaultLayout.ini in Perforce.
+	// 2. Change the code below as you wish and compile the code.
+	// 3. (Optional:) Save your current layout so you can load it later.
+	// 4. Close the editor.
+	// 5. Manually remove Engine\Saved\Config\Windows\EditorLayout.ini
+	// 6. Open the Editor, which will auto-regenerate a default EditorLayout.ini that uses your new code below.
+	// 7. "Window" --> "Save Layout" --> "Save Layout As..."
+	//     - Name: Default Editor Layout
+	//     - Description: Default layout that the Unreal Editor automatically generates
+	// 8. Either click on the toast generated by Unreal that would open the saving path or manually open Engine\Saved\Config\Layouts\ in your explorer
+	// 9. Move and rename the new file (Engine\Saved\Config\Layouts\Default_Editor_Layout.ini) into Engine\Config\Layouts\DefaultLayout.ini
+	// 10. Push the new "DefaultLayout.ini" together with your new code.
+	// 11. Also update these instructions if you change the version number (e.g., from "UnrealEd_Layout_v1.4" to "UnrealEd_Layout_v1.5").
 	const TSharedRef<FTabManager::FLayout> Layout = FLayoutSaveRestore::LoadFromConfig(GEditorLayoutIni,
 		FTabManager::NewLayout( "LevelEditor_Layout_v1.1" )
 		->AddArea
diff --git a/Engine/Source/Editor/MainFrame/Private/MainFrameModule.cpp b/Engine/Source/Editor/MainFrame/Private/MainFrameModule.cpp
--- a/Engine/Source/Editor/MainFrame/Private/MainFrameModule.cpp
+++ b/Engine/Source/Editor/MainFrame/Private/MainFrameModule.cpp
@@ -144,6 +144,20 @@ void FMainFrameModule::CreateDefaultMainFrame( const bool bStartImmersive, const
 				DisplayMetrics.PrimaryDisplayWorkAreaRect.Bottom - DisplayMetrics.PrimaryDisplayWorkAreaRect.Top );
 			const FVector2D WindowSize = (CenterScale * DisplaySize) / DPIScale;
 
+			// IMPORTANT: If you want to change the default value of "LevelEditor_Layout_v1.1" or "UnrealEd_Layout_v1.4" (even if you only change their version numbers), these are the steps to follow:
+			// 1. Check out Engine\Config\Layouts\DefaultLayout.ini in Perforce.
+			// 2. Change the code below as you wish and compile the code.
+			// 3. (Optional:) Save your current layout so you can load it later.
+			// 4. Close the editor.
+			// 5. Manually remove Engine\Saved\Config\Windows\EditorLayout.ini
+			// 6. Open the Editor, which will auto-regenerate a default EditorLayout.ini that uses your new code below.
+			// 7. "Window" --> "Save Layout" --> "Save Layout As..."
+			//     - Name: Default Editor Layout
+			//     - Description: Default layout that the Unreal Editor automatically generates
+			// 8. Either click on the toast generated by Unreal that would open the saving path or manually open Engine\Saved\Config\Layouts\ in your explorer
+			// 9. Move and rename the new file (Engine\Saved\Config\Layouts\Default_Editor_Layout.ini) into Engine\Config\Layouts\DefaultLayout.ini
+			// 10. Push the new "DefaultLayout.ini" together with your new code.
+			// 11. Also update these instructions if you change the version number (e.g., from "UnrealEd_Layout_v1.4" to "UnrealEd_Layout_v1.5").
 			TSharedRef<FTabManager::FLayout> LoadedLayout = FLayoutSaveRestore::LoadFromConfig(GEditorLayoutIni,
 				// We persist the positioning of the level editor and the content browser.
 				// The asset editors currently do not get saved.
