Bug fixed: Editor used to stall on Mac after key up, mouse up. Real issue: Every time a key is pressed on Mac, all menus are re-generated, making it very slow. Workaround: Disabled capability to know whether a UI layout was being currently selected to speed up the generation of the 3 layout menus.

diff --git a/Engine/Source/Editor/MainFrame/Private/Menus/LayoutsMenu.cpp b/Engine/Source/Editor/MainFrame/Private/Menus/LayoutsMenu.cpp
--- a/Engine/Source/Editor/MainFrame/Private/Menus/LayoutsMenu.cpp
+++ b/Engine/Source/Editor/MainFrame/Private/Menus/LayoutsMenu.cpp
@@ -386,8 +386,10 @@ bool IsLayoutCheckedInternal(const FString& InLayoutFullPath, const bool bCheckT
 
 void MakeXLayoutsMenuInternal(UToolMenu* InToolMenu, const TArray<TSharedPtr<FUICommandInfo>>& InXLayoutCommands, const TArray<TSharedPtr<FUICommandInfo>>& InXUserLayoutCommands, const bool bDisplayDefaultLayouts)
 {
+#if !PLATFORM_MAC // On Mac, each time a key is pressed, all menus are re-generated, stalling the Editor given that SaveLayout is slow on Mac because it does not caches as in Windows.
 	// Update GEditorLayoutIni file. Otherwise, we could not track the changes the user did since the layout was loaded
 	FLayoutsMenuSave::SaveLayout();
+#endif
 	// UE Default Layouts
 	if (bDisplayDefaultLayouts)
 	{
@@ -1122,16 +1124,24 @@ bool FLayoutsMenuBase::IsThereUserLayouts()
 
 bool FLayoutsMenuBase::IsLayoutChecked(const int32 InLayoutIndex)
 {
+#if PLATFORM_MAC // On Mac, each time a key is pressed, all menus are re-generated, stalling the Editor given that SaveLayout is slow on Mac because it does not caches as in Windows.
+	return false;
+#else
 	// Check if the desired layout file matches the one currently loaded
 	const bool bCheckTempFileToo = true;
 	return IsLayoutCheckedInternal(GetLayout(InLayoutIndex), bCheckTempFileToo);
+#endif
 }
 
 bool FLayoutsMenuBase::IsUserLayoutChecked(const int32 InLayoutIndex)
 {
+#if PLATFORM_MAC // On Mac, each time a key is pressed, all menus are re-generated, stalling the Editor given that SaveLayout is slow on Mac because it does not caches as in Windows.
+	return false;
+#else
 	// Check if the desired layout file matches the one currently loaded
 	const bool bCheckTempFileToo = true;
 	return IsLayoutCheckedInternal(GetUserLayout(InLayoutIndex), bCheckTempFileToo);
+#endif
 }
 
 #undef LOCTEXT_NAMESPACE
