UE4: Expose support for overriding the application title

Added new IMainFrame module interface function to override the application title.

IMainFrameModule::SetOverrideApplicationTitle(FText)

diff --git a/Engine/Source/Editor/MainFrame/Private/Frame/MainFrameHandler.h b/Engine/Source/Editor/MainFrame/Private/Frame/MainFrameHandler.h
--- a/Engine/Source/Editor/MainFrame/Private/Frame/MainFrameHandler.h
+++ b/Engine/Source/Editor/MainFrame/Private/Frame/MainFrameHandler.h
@@ -185,8 +185,10 @@ public:
 		// Persistent layouts should get stored using the specified method.
 		GlobalTabManager->SetOnPersistLayout(FTabManager::FOnPersistLayout::CreateRaw(this, &FMainFrameHandler::HandleTabManagerPersistLayout));
 		
+		IMainFrameModule& MainFrameModule = FModuleManager::LoadModuleChecked<IMainFrameModule>(TEXT("MainFrame"));
+
 		const bool bIncludeGameName = true;
-		GlobalTabManager->SetApplicationTitle( StaticGetApplicationTitle( bIncludeGameName ) );
+		GlobalTabManager->SetApplicationTitle( MainFrameModule.GetApplicationTitle( bIncludeGameName ) );
 		
 		InRootWindow->SetRequestDestroyWindowOverride( FRequestDestroyWindowOverride::CreateRaw( this, &FMainFrameHandler::CloseRootWindowOverride ) );
 
diff --git a/Engine/Source/Editor/MainFrame/Private/MainFrameModule.cpp b/Engine/Source/Editor/MainFrame/Private/MainFrameModule.cpp
--- a/Engine/Source/Editor/MainFrame/Private/MainFrameModule.cpp
+++ b/Engine/Source/Editor/MainFrame/Private/MainFrameModule.cpp
@@ -430,6 +430,12 @@ void FMainFrameModule::SetLevelNameForWindowTitle( const FString& InLevelFileNam
 		: NSLOCTEXT("UnrealEd", "Untitled", "Untitled" ).ToString();
 }
 
+void FMainFrameModule::SetApplicationTitleOverride(const FText& NewOverriddenApplicationTitle)
+{
+	OverriddenWindowTitle = NewOverriddenApplicationTitle;
+}
+
+
 
 /* IModuleInterface implementation
  *****************************************************************************/
diff --git a/Engine/Source/Editor/MainFrame/Private/MainFrameModule.h b/Engine/Source/Editor/MainFrame/Private/MainFrameModule.h
--- a/Engine/Source/Editor/MainFrame/Private/MainFrameModule.h
+++ b/Engine/Source/Editor/MainFrame/Private/MainFrameModule.h
@@ -84,9 +84,11 @@ public:
 
 	virtual const FText GetApplicationTitle( const bool bIncludeGameName ) const override
 	{
-		return StaticGetApplicationTitle( bIncludeGameName );
+		return OverriddenWindowTitle.IsEmpty() ? StaticGetApplicationTitle( bIncludeGameName ) : OverriddenWindowTitle;
 	}
 
+	virtual void SetApplicationTitleOverride(const FText& NewOverriddenApplicationTitle) override;
+
 	virtual void ShowAboutWindow( ) const override
 	{
 		FMainFrameActionCallbacks::AboutUnrealEd_Execute();
@@ -167,6 +169,9 @@ private:
 	// Friendly name for persistently level name currently loaded.  Used for window and tab titles.
 	FString LoadedLevelName;
 
+	// Override window title, or empty to not override
+	FText OverriddenWindowTitle;
+
 	/// Event to be called when the mainframe is fully created.
 	FMainFrameCreationFinishedEvent MainFrameCreationFinishedEvent;
 
diff --git a/Engine/Source/Editor/MainFrame/Public/Interfaces/IMainFrameModule.h b/Engine/Source/Editor/MainFrame/Public/Interfaces/IMainFrameModule.h
--- a/Engine/Source/Editor/MainFrame/Public/Interfaces/IMainFrameModule.h
+++ b/Engine/Source/Editor/MainFrame/Public/Interfaces/IMainFrameModule.h
@@ -128,6 +128,13 @@ public:
 	 */
 	virtual void SetLevelNameForWindowTitle(const FString& InLevelFileName) = 0;
 
+	/**
+	 * Overrides the title of the application that's displayed in the title bar area and other locations
+	 *
+	 * @param	NewOverriddenApplicationTitle	The text to be displayed in the window title, or empty to use the application's default text
+	 */
+	virtual void SetApplicationTitleOverride(const FText& NewOverriddenApplicationTitle) = 0;
+
 	/**
 	 * Returns a friendly string name for the currently loaded persistent level.
 	 *
