UE4: Expose support for overriding the application title (part 2)

Fixed non-unity compile error

diff --git a/Engine/Source/Editor/MainFrame/Private/Frame/MainFrameHandler.h b/Engine/Source/Editor/MainFrame/Private/Frame/MainFrameHandler.h
--- a/Engine/Source/Editor/MainFrame/Private/Frame/MainFrameHandler.h
+++ b/Engine/Source/Editor/MainFrame/Private/Frame/MainFrameHandler.h
@@ -11,7 +11,7 @@
 #include "Framework/Docking/TabManager.h"
 #include "Framework/Docking/LayoutService.h"
 #include "EngineGlobals.h"
-
+#include "Interfaces/IMainFrameModule.h"
 #include "Editor/UnrealEdEngine.h"
 #include "EditorModeManager.h"
 #include "EditorModes.h"
