Add tooltips to Save As error labels

diff --git a/Engine/Source/Editor/ContentBrowser/Private/SAssetDialog.cpp b/Engine/Source/Editor/ContentBrowser/Private/SAssetDialog.cpp
--- a/Engine/Source/Editor/ContentBrowser/Private/SAssetDialog.cpp
+++ b/Engine/Source/Editor/ContentBrowser/Private/SAssetDialog.cpp
@@ -176,6 +176,7 @@ void SAssetDialog::Construct(const FArguments& InArgs, const FSharedAssetDialogC
 				[
 					SNew(STextBlock)
 					.Text( this, &SAssetDialog::GetNameErrorLabelText )
+					.ToolTipText(this, &SAssetDialog::GetNameErrorLabelText)
 					.TextStyle( FEditorStyle::Get(), "AssetDialog.ErrorLabelFont" )
 				]
 			]
diff --git a/Engine/Source/Editor/MainFrame/Private/Menus/SaveLayoutDialog.cpp b/Engine/Source/Editor/MainFrame/Private/Menus/SaveLayoutDialog.cpp
--- a/Engine/Source/Editor/MainFrame/Private/Menus/SaveLayoutDialog.cpp
+++ b/Engine/Source/Editor/MainFrame/Private/Menus/SaveLayoutDialog.cpp
@@ -268,6 +268,7 @@ void SSaveLayoutDialog::Construct(const FArguments& InArgs, const TSharedRef<FSa
 			[
 				SNew(STextBlock)
 				.Text( this, &SSaveLayoutDialog::GetNameErrorLabelText )
+				.ToolTipText(this, &SSaveLayoutDialog::GetNameErrorLabelText)
 				.TextStyle( FEditorStyle::Get(), "AssetDialog.ErrorLabelFont" )
 			]
 		]
