Avoid creating 300 empty keyboard bindings without description.

Thanks matt for the initial CL!

diff --git a/Engine/Source/Editor/MainFrame/Private/Frame/MainFrameActions.cpp b/Engine/Source/Editor/MainFrame/Private/Frame/MainFrameActions.cpp
--- a/Engine/Source/Editor/MainFrame/Private/Frame/MainFrameActions.cpp
+++ b/Engine/Source/Editor/MainFrame/Private/Frame/MainFrameActions.cpp
@@ -195,7 +195,22 @@ void FMainFrameCommands::RegisterCommands()
 	ActionList->MapAction( CreditsUnrealEd, FExecuteAction::CreateStatic(&FMainFrameActionCallbacks::CreditsUnrealEd_Execute) );
 
 	// Layout commands
-	RegisterLayoutCommands();
+	UI_COMMAND(ImportLayout, "Import Layout...", "Import a custom layout (or set of layouts) from a different directory and load it into your current instance of the Unreal Editor UI", EUserInterfaceActionType::Button, FInputChord());
+	ActionList->MapAction(ImportLayout, FExecuteAction::CreateStatic(&FLayoutsMenuLoad::ImportLayout));
+
+	UI_COMMAND(SaveLayoutAs, "Save Layout As...", "Save the current layout customization on disk so it can be loaded later", EUserInterfaceActionType::Button, FInputChord());
+	ActionList->MapAction(SaveLayoutAs, FExecuteAction::CreateStatic(&FLayoutsMenuSave::SaveLayoutAs));
+
+	UI_COMMAND(ExportLayout, "Export Layout...", "Export the custom layout customization to a different directory", EUserInterfaceActionType::Button, FInputChord());
+	ActionList->MapAction(ExportLayout, FExecuteAction::CreateStatic(&FLayoutsMenuSave::ExportLayout));
+
+	UI_COMMAND(RemoveUserLayouts, "Remove All User Layouts...", "Remove all the layout customizations created by the user", EUserInterfaceActionType::Button, FInputChord());
+	ActionList->MapAction(
+		RemoveUserLayouts,
+		FExecuteAction::CreateStatic(&FLayoutsMenuRemove::RemoveUserLayouts),
+		FCanExecuteAction::CreateStatic(&FLayoutsMenuBase::IsThereUserLayouts)
+	);
+
 
 #if !PLATFORM_MAC && !PLATFORM_LINUX // Fullscreen mode in the editor is currently unsupported on Mac and Linux
 	UI_COMMAND( ToggleFullscreen, "Enable Fullscreen", "Enables fullscreen mode for the application, expanding across the entire monitor", EUserInterfaceActionType::ToggleButton, FInputChord(EModifierKey::Shift, EKeys::F11) );
@@ -212,96 +227,6 @@ void FMainFrameCommands::RegisterCommands()
 	FGlobalEditorCommonCommands::MapActions(ActionList);
 }
 
-TSharedRef<FUICommandInfo> CreateFUICommandInfoInternal(const TSharedRef<FBindingContext>& InThisAsShared, const FName& InCommandName)
-{
-	// Default FUICommandInfoDecl command
-	return FUICommandInfoDecl(
-			InThisAsShared,
-			InCommandName,
-			FText::FromString("InLabel"),
-			FText::FromString("InDesc"))
-		.UserInterfaceType(EUserInterfaceActionType::Check)
-		.DefaultChord(FInputChord());
-}
-
-void FMainFrameCommands::RegisterLayoutCommands()
-{
-	// Load layout commands
-	UI_COMMAND(MainFrameLayoutCommands.ImportLayout, "Import Layout...", "Import a custom layout (or set of layouts) from a different directory and load it into your current instance of the Unreal Editor UI", EUserInterfaceActionType::Button, FInputChord());
-	ActionList->MapAction(MainFrameLayoutCommands.ImportLayout, FExecuteAction::CreateStatic(&FLayoutsMenuLoad::ImportLayout));
-
-	// Save layout commands
-	UI_COMMAND(MainFrameLayoutCommands.SaveLayoutAs, "Save Layout As...", "Save the current layout customization on disk so it can be loaded later", EUserInterfaceActionType::Button, FInputChord());
-	ActionList->MapAction(MainFrameLayoutCommands.SaveLayoutAs, FExecuteAction::CreateStatic(&FLayoutsMenuSave::SaveLayoutAs));
-	UI_COMMAND(MainFrameLayoutCommands.ExportLayout, "Export Layout...", "Export the custom layout customization to a different directory", EUserInterfaceActionType::Button, FInputChord());
-	ActionList->MapAction(MainFrameLayoutCommands.ExportLayout, FExecuteAction::CreateStatic(&FLayoutsMenuSave::ExportLayout));
-
-	// Remove layout commands
-	UI_COMMAND(MainFrameLayoutCommands.RemoveUserLayouts, "Remove All User Layouts...", "Remove all the layout customizations created by the user", EUserInterfaceActionType::Button, FInputChord());
-	ActionList->MapAction(
-		MainFrameLayoutCommands.RemoveUserLayouts,
-		FExecuteAction::CreateStatic(&FLayoutsMenuRemove::RemoveUserLayouts),
-		FCanExecuteAction::CreateStatic(&FLayoutsMenuBase::IsThereUserLayouts)
-	);
-
-	// Load/Override/Remove layout commands
-	const int32 MaxLayouts = 50;
-	for (int32 CurrentLayoutIndex = 0; CurrentLayoutIndex < MaxLayouts; ++CurrentLayoutIndex)
-	{
-		// NOTE: The actual label and tool-tip will be overridden at runtime when the command is bound to a menu item, however
-		// we still need to set one here so that the key bindings UI can function properly
-		// Thus, we remove any (NS)LOCTEXT from here to avoid unnecessary translations.
-
-		// Load Layout
-		{
-			// Default Layouts
-			TSharedRef<FUICommandInfo> LoadLayout = CreateFUICommandInfoInternal(this->AsShared(), FName(*FString::Printf(TEXT("LoadLayout%i"), CurrentLayoutIndex)));
-			MainFrameLayoutCommands.LoadLayoutCommands.Add(LoadLayout);
-			ActionList->MapAction(MainFrameLayoutCommands.LoadLayoutCommands[CurrentLayoutIndex], FExecuteAction::CreateStatic(&FLayoutsMenuLoad::LoadLayout, CurrentLayoutIndex),
-				FCanExecuteAction::CreateStatic(&FLayoutsMenuLoad::CanLoadChooseLayout, CurrentLayoutIndex),
-				FIsActionChecked::CreateStatic(&FLayoutsMenuBase::IsLayoutChecked, CurrentLayoutIndex));
-			// User Layouts
-			TSharedRef<FUICommandInfo> LoadUserLayout = CreateFUICommandInfoInternal(this->AsShared(), FName(*FString::Printf(TEXT("LoadUserLayout%i"), CurrentLayoutIndex)));
-			MainFrameLayoutCommands.LoadUserLayoutCommands.Add(LoadUserLayout);
-			ActionList->MapAction(MainFrameLayoutCommands.LoadUserLayoutCommands[CurrentLayoutIndex], FExecuteAction::CreateStatic(&FLayoutsMenuLoad::LoadUserLayout, CurrentLayoutIndex),
-				FCanExecuteAction::CreateStatic(&FLayoutsMenuLoad::CanLoadChooseUserLayout, CurrentLayoutIndex),
-				FIsActionChecked::CreateStatic(&FLayoutsMenuBase::IsUserLayoutChecked, CurrentLayoutIndex));
-		}
-
-		// Override Layout
-		{
-			// Default Layouts
-			TSharedRef<FUICommandInfo> OverrideLayout = CreateFUICommandInfoInternal(this->AsShared(), FName(*FString::Printf(TEXT("OverrideLayout%i"), CurrentLayoutIndex)));
-			MainFrameLayoutCommands.OverrideLayoutCommands.Add(OverrideLayout);
-			ActionList->MapAction(MainFrameLayoutCommands.OverrideLayoutCommands[CurrentLayoutIndex], FExecuteAction::CreateStatic(&FLayoutsMenuSave::OverrideLayout, CurrentLayoutIndex),
-				FCanExecuteAction::CreateStatic(&FLayoutsMenuSave::CanSaveChooseLayout, CurrentLayoutIndex),
-				FIsActionChecked::CreateStatic(&FLayoutsMenuBase::IsLayoutChecked, CurrentLayoutIndex));
-			// User Layouts
-			TSharedRef<FUICommandInfo> OverrideUserLayout = CreateFUICommandInfoInternal(this->AsShared(), FName(*FString::Printf(TEXT("OverrideUserLayout%i"), CurrentLayoutIndex)));
-			MainFrameLayoutCommands.OverrideUserLayoutCommands.Add(OverrideUserLayout);
-			ActionList->MapAction(MainFrameLayoutCommands.OverrideUserLayoutCommands[CurrentLayoutIndex], FExecuteAction::CreateStatic(&FLayoutsMenuSave::OverrideUserLayout, CurrentLayoutIndex),
-				FCanExecuteAction::CreateStatic(&FLayoutsMenuSave::CanSaveChooseUserLayout, CurrentLayoutIndex),
-				FIsActionChecked::CreateStatic(&FLayoutsMenuBase::IsUserLayoutChecked, CurrentLayoutIndex));
-		}
-
-		// Remove Layout
-		{
-			// Default Layouts
-			TSharedRef<FUICommandInfo> RemoveLayout = CreateFUICommandInfoInternal(this->AsShared(), FName(*FString::Printf(TEXT("RemoveLayout%i"), CurrentLayoutIndex)));
-			MainFrameLayoutCommands.RemoveLayoutCommands.Add(RemoveLayout);
-			ActionList->MapAction(MainFrameLayoutCommands.RemoveLayoutCommands[CurrentLayoutIndex], FExecuteAction::CreateStatic(&FLayoutsMenuRemove::RemoveLayout, CurrentLayoutIndex),
-				FCanExecuteAction::CreateStatic(&FLayoutsMenuRemove::CanRemoveChooseLayout, CurrentLayoutIndex),
-				FIsActionChecked::CreateStatic(&FLayoutsMenuBase::IsLayoutChecked, CurrentLayoutIndex));
-			// User Layouts
-			TSharedRef<FUICommandInfo> RemoveUserLayout = CreateFUICommandInfoInternal(this->AsShared(), FName(*FString::Printf(TEXT("RemoveUserLayout%i"), CurrentLayoutIndex)));
-			MainFrameLayoutCommands.RemoveUserLayoutCommands.Add(RemoveUserLayout);
-			ActionList->MapAction(MainFrameLayoutCommands.RemoveUserLayoutCommands[CurrentLayoutIndex], FExecuteAction::CreateStatic(&FLayoutsMenuRemove::RemoveUserLayout, CurrentLayoutIndex),
-				FCanExecuteAction::CreateStatic(&FLayoutsMenuRemove::CanRemoveChooseUserLayout, CurrentLayoutIndex),
-				FIsActionChecked::CreateStatic(&FLayoutsMenuBase::IsUserLayoutChecked, CurrentLayoutIndex));
-		}
-	}
-}
-
 FReply FMainFrameActionCallbacks::OnUnhandledKeyDownEvent(const FKeyEvent& InKeyEvent)
 {
 	if(!GIsSlowTask)
diff --git a/Engine/Source/Editor/MainFrame/Private/Frame/MainFrameActions.h b/Engine/Source/Editor/MainFrame/Private/Frame/MainFrameActions.h
--- a/Engine/Source/Editor/MainFrame/Private/Frame/MainFrameActions.h
+++ b/Engine/Source/Editor/MainFrame/Private/Frame/MainFrameActions.h
@@ -11,26 +11,6 @@
 
 class FUICommandList;
 
-/**
- * Unreal editor main frame actions related to the Layout section
- */
-class FMainFrameLayoutCommands
-{
-public:
-	// Load-based commands
-	TArray<TSharedPtr<FUICommandInfo>> LoadLayoutCommands;
-	TArray<TSharedPtr<FUICommandInfo>> LoadUserLayoutCommands;
-	TSharedPtr< FUICommandInfo > ImportLayout;
-	// Save-based commands
-	TArray<TSharedPtr<FUICommandInfo>> OverrideLayoutCommands;
-	TArray<TSharedPtr<FUICommandInfo>> OverrideUserLayoutCommands;
-	TSharedPtr< FUICommandInfo > SaveLayoutAs;
-	TSharedPtr< FUICommandInfo > ExportLayout;
-	// Remove-based commands
-	TArray<TSharedPtr<FUICommandInfo>> RemoveLayoutCommands;
-	TArray<TSharedPtr<FUICommandInfo>> RemoveUserLayoutCommands;
-	TSharedPtr< FUICommandInfo > RemoveUserLayouts;
-};
 
 /**
  * Unreal editor main frame actions
@@ -83,7 +63,10 @@ public:
 	TSharedPtr< FUICommandInfo > CreditsUnrealEd;
 
 	// Layout
-	FMainFrameLayoutCommands MainFrameLayoutCommands;
+	TSharedPtr< FUICommandInfo > ImportLayout;
+	TSharedPtr< FUICommandInfo > SaveLayoutAs;
+	TSharedPtr< FUICommandInfo > ExportLayout;
+	TSharedPtr< FUICommandInfo > RemoveUserLayouts;
 
 	TSharedPtr< FUICommandInfo > ToggleFullscreen;
 
diff --git a/Engine/Source/Editor/MainFrame/Private/Menus/LayoutsMenu.cpp b/Engine/Source/Editor/MainFrame/Private/Menus/LayoutsMenu.cpp
--- a/Engine/Source/Editor/MainFrame/Private/Menus/LayoutsMenu.cpp
+++ b/Engine/Source/Editor/MainFrame/Private/Menus/LayoutsMenu.cpp
@@ -39,6 +39,21 @@
 
 DEFINE_LOG_CATEGORY_STATIC(LogLayoutsMenu, Fatal, All);
 
+namespace FPrivateFLayoutsMenu
+{
+	enum class ELayoutsMenu
+	{
+		Load,
+		Save,
+		Remove
+	};
+
+	enum class ELayoutsType
+	{
+		Engine,
+		User
+	};
+
 // Get LayoutsDirectory path
 FString CreateAndGetDefaultLayoutDirInternal()
 {
@@ -112,7 +127,8 @@ bool TrySaveLayoutOrWarnInternal(
 							" - Current layout description: {1}\n\n"
 							"If you select \"Preserve Values\", the displayed name and description of the original layout customization will also be copied into the new configuration file.\n\n"
 							"If you select \"Clear Values\", these fields will be emptied.\n\n"
-							"If you are not sure, select \"Preserve Values\" if you are exporting the layout configuration without making any changes, or \"Clear Values\" if you have made or plan to make changes to the layout.\n\n"
+							"If you are not sure, select \"Preserve Values\" if you are exporting the layout configuration without making any changes, or \"Clear Values\" if you"" have made or"
+							" plan to make changes to the layout.\n\n"
 						),
 						LayoutNameSource, LayoutDescriptionSource);
 					// Dialog SWidget
@@ -173,14 +189,16 @@ bool TrySaveLayoutOrWarnInternal(
 		if (!FPaths::FileExists(InSourceFilePath))
 		{
 			Arguments.Add(TEXT("FileName"), FText::FromString(FPaths::ConvertRelativePathToFull(InSourceFilePath)));
-			TextBody = FText::Format(LOCTEXT("UnsuccessfulSave_NoExist_Notification", "The requested operation ({WhatIs}) was unsuccessful, the desired file does not exist. File path:\n{FileName}"), Arguments);
+			TextBody = FText::Format(LOCTEXT("UnsuccessfulSave_NoExist_Notification",
+				"The requested operation ({WhatIs}) was unsuccessful, the desired file does not exist. File path:\n{FileName}"), Arguments);
 			EditorErrors.Warning(TextBody);
 		}
 		// Target is read-only
 		else if (IFileManager::Get().IsReadOnly(*InTargetFilePath))
 		{
 			Arguments.Add(TEXT("FileName"), FText::FromString(TargetAbsoluteFilePath));
-			TextBody = FText::Format(LOCTEXT("UnsuccessfulSave_ReadOnly_Notification", "The requested operation ({WhatIs}) was unsuccessful, the target file path is read-only. File path:\n{FileName}"), Arguments);
+			TextBody = FText::Format(LOCTEXT("UnsuccessfulSave_ReadOnly_Notification",
+				"The requested operation ({WhatIs}) was unsuccessful, the target file path is read-only. File path:\n{FileName}"), Arguments);
 			EditorErrors.Warning(TextBody);
 		}
 		// Target and source are the same
@@ -189,7 +207,8 @@ bool TrySaveLayoutOrWarnInternal(
 			Arguments.Add(TEXT("SourceFileName"), FText::FromString(FPaths::ConvertRelativePathToFull(InSourceFilePath)));
 			Arguments.Add(TEXT("FinalFileName"), FText::FromString(TargetAbsoluteFilePath));
 			TextBody = FText::Format(LOCTEXT("UnsuccessfulSave_Fallback_Notification",
-				"The requested operation ({WhatIs}) was unsuccessful, target and source layout file paths are the same ({SourceFileName})!\nAre you trying to import or replace a file that is already in the layouts folder? If so, remove the current file first."), Arguments);
+				"The requested operation ({WhatIs}) was unsuccessful, target and source layout file paths are the same ({SourceFileName})!\nAre you trying to import or replace a file"
+				" that is already in the layouts folder? If so, remove the current file first."), Arguments);
 			EditorErrors.Warning(TextBody);
 		}
 		// We don't specifically know why it failed, this is a fallback
@@ -197,7 +216,9 @@ bool TrySaveLayoutOrWarnInternal(
 		{
 			Arguments.Add(TEXT("SourceFileName"), FText::FromString(FPaths::ConvertRelativePathToFull(InSourceFilePath)));
 			Arguments.Add(TEXT("FinalFileName"), FText::FromString(TargetAbsoluteFilePath));
-			TextBody = FText::Format(LOCTEXT("UnsuccessfulSave_Fallback_Notification", "The requested operation ({WhatIs}) was unsuccessful while copying the layout file from\n{SourceFileName}\ninto\n{FinalFileName}\n\nUsually, this occurs when the introduced file name contains unsupported characters or the total path length exceeds the OS limit."), Arguments);
+			TextBody = FText::Format(LOCTEXT("UnsuccessfulSave_Fallback_Notification",
+				"The requested operation ({WhatIs}) was unsuccessful while copying the layout file from\n{SourceFileName}\ninto\n{FinalFileName}\n\nUsually, this occurs when the"
+				" introduced file name contains unsupported characters or the total path length exceeds the OS limit."), Arguments);
 			EditorErrors.Warning(TextBody);
 		}
 		EditorErrors.Notify(LOCTEXT("LoadUnsuccessful_Title", "Load Unsuccessful!"));
@@ -278,10 +299,10 @@ FText GetTooltipTextInternal(const FText& InDisplayName, const FString& InLayout
 	return Tooltip;
 }
 
-void DisplayLayoutsInternal(FToolMenuSection& InSection, const TArray<TSharedPtr<FUICommandInfo>>& InXLayoutCommands, const TArray<FString> InLayoutIniFileNames, const FString InLayoutsDirectory)
+void DisplayLayoutsInternal(FToolMenuSection& InSection, const TArray<FString>& InLayoutIniFileNames, const FString& InLayoutsDirectory, const ELayoutsMenu InLayoutsMenu, const ELayoutsType InLayoutsType)
 {
 	// If there are Layout ini files, read them
-	for (int32 LayoutIndex = 0; LayoutIndex < InLayoutIniFileNames.Num() && LayoutIndex < InXLayoutCommands.Num(); ++LayoutIndex)
+	for (int32 LayoutIndex = 0; LayoutIndex < InLayoutIniFileNames.Num(); ++LayoutIndex)
 	{
 		const FString LayoutFilePath = IFileManager::Get().ConvertToAbsolutePathForExternalAppForRead(*FPaths::Combine(InLayoutsDirectory, InLayoutIniFileNames[LayoutIndex]));
 		// Make sure it is a layout file
@@ -294,7 +315,51 @@ void DisplayLayoutsInternal(FToolMenuSection& InSection, const TArray<TSharedPtr
 			// If no localization name, then display the file name
 			const FText DisplayName = (!LayoutName.IsEmpty() ? LayoutName : GetDisplayTextInternal(InLayoutIniFileNames[LayoutIndex]));
 			const FText Tooltip = GetTooltipTextInternal(DisplayName, LayoutFilePath, LayoutDescription);
-			InSection.AddMenuEntry(InXLayoutCommands[LayoutIndex], DisplayName, Tooltip).Name = NAME_None;
+
+			// Create UI action here that calls the necessary code in FLayoutsMenuLoad, Save, or Remove 
+			FUIAction UIAction;
+			if (InLayoutsMenu == ELayoutsMenu::Load && InLayoutsType == ELayoutsType::Engine)
+			{
+				UIAction = FUIAction(FExecuteAction::CreateStatic(&FLayoutsMenuLoad::LoadLayout, LayoutIndex),
+					FCanExecuteAction::CreateStatic(&FLayoutsMenuLoad::CanLoadChooseLayout, LayoutIndex),
+					FIsActionChecked::CreateStatic(&FLayoutsMenuBase::IsLayoutChecked, LayoutIndex));
+			}
+			else if (InLayoutsMenu == ELayoutsMenu::Load && InLayoutsType == ELayoutsType::User)
+			{
+				UIAction = FUIAction(FExecuteAction::CreateStatic(&FLayoutsMenuLoad::LoadUserLayout, LayoutIndex),
+					FCanExecuteAction::CreateStatic(&FLayoutsMenuLoad::CanLoadChooseUserLayout, LayoutIndex),
+					FIsActionChecked::CreateStatic(&FLayoutsMenuBase::IsUserLayoutChecked, LayoutIndex));
+			}
+			else if (InLayoutsMenu == ELayoutsMenu::Save && InLayoutsType == ELayoutsType::Engine)
+			{
+				UIAction = FUIAction(FExecuteAction::CreateStatic(&FLayoutsMenuSave::OverrideLayout, LayoutIndex),
+					FCanExecuteAction::CreateStatic(&FLayoutsMenuSave::CanSaveChooseLayout, LayoutIndex),
+					FIsActionChecked::CreateStatic(&FLayoutsMenuBase::IsLayoutChecked, LayoutIndex));
+			}
+			else if (InLayoutsMenu == ELayoutsMenu::Save && InLayoutsType == ELayoutsType::User)
+			{
+				UIAction = FUIAction(FExecuteAction::CreateStatic(&FLayoutsMenuSave::OverrideUserLayout, LayoutIndex),
+					FCanExecuteAction::CreateStatic(&FLayoutsMenuSave::CanSaveChooseUserLayout, LayoutIndex),
+					FIsActionChecked::CreateStatic(&FLayoutsMenuBase::IsUserLayoutChecked, LayoutIndex));
+			}
+			else if (InLayoutsMenu == ELayoutsMenu::Remove && InLayoutsType == ELayoutsType::Engine)
+			{
+				UIAction = FUIAction(FExecuteAction::CreateStatic(&FLayoutsMenuRemove::RemoveLayout, LayoutIndex),
+					FCanExecuteAction::CreateStatic(&FLayoutsMenuRemove::CanRemoveChooseLayout, LayoutIndex),
+					FIsActionChecked::CreateStatic(&FLayoutsMenuBase::IsLayoutChecked, LayoutIndex));
+			}
+			else if (InLayoutsMenu == ELayoutsMenu::Remove && InLayoutsType == ELayoutsType::User)
+			{
+				UIAction = FUIAction(FExecuteAction::CreateStatic(&FLayoutsMenuRemove::RemoveUserLayout, LayoutIndex),
+					FCanExecuteAction::CreateStatic(&FLayoutsMenuRemove::CanRemoveChooseUserLayout, LayoutIndex),
+					FIsActionChecked::CreateStatic(&FLayoutsMenuBase::IsUserLayoutChecked, LayoutIndex));
+			}
+			else
+			{
+				ensure(false);
+			}
+			// Run desired action
+			InSection.AddMenuEntry(NAME_None, DisplayName, Tooltip, FSlateIcon(), FToolUIActionChoice(UIAction));
 		}
 	}
 }
@@ -384,7 +449,7 @@ bool IsLayoutCheckedInternal(const FString& InLayoutFullPath, const bool bCheckT
 	}
 }
 
-void MakeXLayoutsMenuInternal(UToolMenu* InToolMenu, const TArray<TSharedPtr<FUICommandInfo>>& InXLayoutCommands, const TArray<TSharedPtr<FUICommandInfo>>& InXUserLayoutCommands, const bool bDisplayDefaultLayouts)
+void MakeXLayoutsMenuInternal(UToolMenu* InToolMenu, const bool bDisplayDefaultLayouts, const ELayoutsMenu InLayoutsMenu)
 {
 #if !PLATFORM_MAC // On Mac, each time a key is pressed, all menus are re-generated, stalling the Editor given that SaveLayout is slow on Mac because it does not caches as in Windows.
 	// Update GEditorLayoutIni file. Otherwise, we could not track the changes the user did since the layout was loaded
@@ -399,7 +464,7 @@ void MakeXLayoutsMenuInternal(UToolMenu* InToolMenu, const TArray<TSharedPtr<FUI
 		// Get Layout init files
 		const TArray<FString> LayoutIniFileNames = GetIniFilesInFolderInternal(LayoutsDirectory);
 		// If there are user Layout ini files, read and display them
-		DisplayLayoutsInternal(Section, InXLayoutCommands, LayoutIniFileNames, LayoutsDirectory);
+		DisplayLayoutsInternal(Section, LayoutIniFileNames, LayoutsDirectory, InLayoutsMenu, ELayoutsType::Engine);
 	}
 	// User Layouts
 	{
@@ -409,7 +474,7 @@ void MakeXLayoutsMenuInternal(UToolMenu* InToolMenu, const TArray<TSharedPtr<FUI
 		// Get User Layout init files
 		const TArray<FString> UserLayoutIniFileNames = GetIniFilesInFolderInternal(UserLayoutsDirectory);
 		// If there are user Layout ini files, read and display them
-		DisplayLayoutsInternal(Section, InXUserLayoutCommands, UserLayoutIniFileNames, UserLayoutsDirectory);
+		DisplayLayoutsInternal(Section, UserLayoutIniFileNames, UserLayoutsDirectory, InLayoutsMenu, ELayoutsType::User);
 	}
 }
 
@@ -512,19 +577,22 @@ bool CheckAskUserAndClosePIESIEAndAssetEditors(const FText& InitialMessage)
 	if (bIsPIEOrSIERunning && bAreAssetEditorOpened)
 	{
 		TextTitle = LOCTEXT("CheckAskUserAndClosePIESIEAndAssetEditorsIfYesHeaderAll", "Close PIE/SIE and Asset Editors?");
-		IfYesText = LOCTEXT("CheckAskUserAndClosePIESIEAndAssetEditorsIfYesBodyAll", "If \"Yes\", your current game instances (PIE or SIE) as well as all open Asset Editors will be closed. Any unsaved changes in those will also be lost.");
+		IfYesText = LOCTEXT("CheckAskUserAndClosePIESIEAndAssetEditorsIfYesBodyAll", "If \"Yes\", your current game instances (PIE or SIE) as well as all open Asset Editors"
+			" will be closed. Any unsaved changes in those will also be lost.");
 	}
 	// If PIE or SIE are opened
 	else if (bIsPIEOrSIERunning)
 	{
 		TextTitle = LOCTEXT("CheckAskUserAndClosePIESIEAndAssetEditorsIfYesHeaderPIE", "Close PIE/SIE?");
-		IfYesText = LOCTEXT("CheckAskUserAndClosePIESIEAndAssetEditorsIfYesBodyPIE", "If \"Yes\", your current game instances (PIE or SIE) will be closed. Any unsaved changes in those will also be lost.");
+		IfYesText = LOCTEXT("CheckAskUserAndClosePIESIEAndAssetEditorsIfYesBodyPIE", "If \"Yes\", your current game instances (PIE or SIE) will be closed. Any unsaved changes"
+			" in those will also be lost.");
 	}
 	// If any Asset Editors is opened
 	else if (bAreAssetEditorOpened)
 	{
 		TextTitle = LOCTEXT("CheckAskUserAndClosePIESIEAndAssetEditorsIfYesHeaderEditorAssets", "Close Asset Editors?");
-		IfYesText = LOCTEXT("CheckAskUserAndClosePIESIEAndAssetEditorsIfYesBodyEditorAssets", "If \"Yes\", all open Asset Editors will be closed. Any unsaved changes in those will also be lost.");
+		IfYesText = LOCTEXT("CheckAskUserAndClosePIESIEAndAssetEditorsIfYesBodyEditorAssets", "If \"Yes\", all open Asset Editors will be closed. Any unsaved changes in those"
+			" will also be lost.");
 	}
 	// FMessageDialog
 	const FText IfNoText = LOCTEXT("CheckAskUserAndClosePIESIEAndAssetEditorsIfNoBody", "If \"No\", you can manually reload the layout from the \"User Layouts\" section later.");
@@ -544,7 +612,8 @@ bool CheckAskUserAndClosePIESIEAndAssetEditors(const FText& InitialMessage)
 		else
 		{
 			ensureMsgf(false,
-				TEXT("This has not been tested because the code does not reach this by default. The layout is loaded through the Editor UI, and GIsPlayInEditorWorld should not have any kind of Editor UI, so it should not be possible to load a layout in that status."));
+				TEXT("This has not been tested because the code does not reach this by default. The layout is loaded through the Editor UI, and GIsPlayInEditorWorld should"
+					" not have any kind of Editor UI, so it should not be possible to load a layout in that status."));
 		}
 	}
 	// If any Asset Editors is opened, ask the user whether he wants to automatically close them and continue loading the layout
@@ -556,13 +625,164 @@ bool CheckAskUserAndClosePIESIEAndAssetEditors(const FText& InitialMessage)
 	return true;
 }
 
+	FText GenerateLocalizedTextForFile(const FText& InText)
+	{
+		// Proper FText to FString
+		FString StringSimulatingText;
+		FTextStringHelper::WriteToBuffer(StringSimulatingText, InText);
+		// Sanitize text (truncate if too big)
+		FString SanitazedTruncatedText = (StringSimulatingText.Len() < 100 ? StringSimulatingText : StringSimulatingText.Left(100));
+		FSaveLayoutDialogUtils::SanitizeText(SanitazedTruncatedText);
+		// Create full file path
+		if (StringSimulatingText.Len() < 10 || StringSimulatingText.Left(9) != TEXT("NSLOCTEXT"))
+		{
+			FString StringSimulatingTextRecreated =
+				// Namespace
+				TEXT("NSLOCTEXT(\"LayoutNamespace\", "
+				// Key
+				"\"") + SanitazedTruncatedText + TEXT("\", "
+				// Source string
+				"\"") + StringSimulatingText + TEXT("\")");
+			return FText::FromString(StringSimulatingTextRecreated);
+		}
+		else
+		{
+			return FText::FromString(StringSimulatingText);
+		}
+	}
+
+	void SaveExportLayoutCommon(const FString& InDefaultDirectory, const bool bMustBeSavedInDefaultDirectory, const FText& InWhatIsThis, const bool bShouldAskBeforeCleaningLayoutNameAndDescriptionFields)
+	{
+		// Export/SaveAs the user-selected layout configuration files and load one of them
+		IDesktopPlatform* DesktopPlatform = FDesktopPlatformModule::Get();
+		if (DesktopPlatform)
+		{
+			bool bWereFilesSelected = false;
+			TArray<FString> LayoutFilePaths;
+			bool bWasDialogOpened = bMustBeSavedInDefaultDirectory;
+			// "Save Layout As..."
+			if (bMustBeSavedInDefaultDirectory)
+			{
+				TArray<FText> LayoutNames;
+				LayoutNames.Emplace(FLayoutSaveRestore::LoadSectionFromConfig(GEditorLayoutIni, "LayoutName"));
+				TArray<FText> LayoutDescriptions;
+				LayoutDescriptions.Emplace(FLayoutSaveRestore::LoadSectionFromConfig(GEditorLayoutIni, "LayoutDescription"));
+				// We want to avoid the duplication of the file/name/description fields, so we add "Copy of " at the beginning of the name and description fields
+				if (LayoutNames[0].ToString().Len() > 0)
+				{
+					LayoutNames[0] = FText::FromString(TEXT("Copy of ") + LayoutNames[0].ToString());
+				}
+				if (LayoutDescriptions[0].ToString().Len() > 0)
+				{
+					LayoutDescriptions[0] = FText::FromString(TEXT("Copy of ") + LayoutDescriptions[0].ToString());
+				}
+				// Create SWidget for saving the layout in its own SWindow and block the thread until it is finished
+				const TSharedRef<FSaveLayoutDialogParams> SaveLayoutDialogParams = MakeShared<FSaveLayoutDialogParams>(InDefaultDirectory, TEXT(".ini"), LayoutNames, LayoutDescriptions);
+				bWasDialogOpened = FSaveLayoutDialogUtils::CreateSaveLayoutAsDialogInStandaloneWindow(SaveLayoutDialogParams);
+				bWereFilesSelected = SaveLayoutDialogParams->bWereFilesSelected;
+				LayoutFilePaths = SaveLayoutDialogParams->LayoutFilePaths;
+				LayoutNames = SaveLayoutDialogParams->LayoutNames;
+				LayoutDescriptions = SaveLayoutDialogParams->LayoutDescriptions;
+
+				// Update GEditorLayoutIni file if LayoutNames or LayoutDescriptions were modified by the user
+				if (bWasDialogOpened && LayoutNames.Num() > 0 && LayoutDescriptions.Num() > 0 && (LayoutNames[0].ToString().Len() > 0 || LayoutDescriptions[0].ToString().Len() > 0))
+				{
+					checkf(LayoutNames.Num() == LayoutDescriptions.Num(), TEXT("There should be the same number of LayoutNames and LayoutDescriptions."));
+					for (int32 Index = 0; Index < LayoutNames.Num(); ++Index)
+					{
+						const FText& LayoutNameAsTextText = GenerateLocalizedTextForFile(LayoutNames[Index]);
+						const FText& LayoutDescriptionAsTextText = GenerateLocalizedTextForFile(LayoutDescriptions[Index]);
+						// Update fields
+						FLayoutSaveRestore::SaveSectionToConfig(GEditorLayoutIni, "LayoutName", LayoutNameAsTextText);
+						FLayoutSaveRestore::SaveSectionToConfig(GEditorLayoutIni, "LayoutDescription", LayoutDescriptionAsTextText);
+						// Flush file
+						const bool bRead = true;
+						GConfig->Flush(bRead, GEditorLayoutIni);
+					}
+				}
+			}
+			// "Export Layout..." (or "Save Layout As..." dialog could not be opened)
+			if (!bWasDialogOpened)
+			{
+				// Open the "save file" dialog so user can save his/her layout configuration file
+				const FString DefaultFile = "";
+				bWereFilesSelected = DesktopPlatform->SaveFileDialog(
+					FSlateApplication::Get().FindBestParentWindowHandleForDialogs(nullptr),
+					TEXT("Export a Layout Configuration File"),
+					InDefaultDirectory,
+					DefaultFile,
+					TEXT("Layout configuration files|*.ini|"),
+					EFileDialogFlags::None, //EFileDialogFlags::Multiple, // Allow/Avoid multiple file selection
+					LayoutFilePaths
+				);
+			}
+			// If file(s) selected, copy them into the user layouts directory and load one of them
+			if (bWereFilesSelected && LayoutFilePaths.Num() > 0)
+			{
+				// Iterate over selected layout ini files
+				FString FirstGoodLayoutFile = "";
+				for (const FString& LayoutFilePath : LayoutFilePaths)
+				{
+					// If writing in the right folder
+					const FString LayoutFilePathAbsolute = IFileManager::Get().ConvertToAbsolutePathForExternalAppForRead(*FPaths::GetPath(LayoutFilePath));
+					const FString DefaultDirectoryAbsolute = IFileManager::Get().ConvertToAbsolutePathForExternalAppForRead(*FPaths::GetPath(InDefaultDirectory));
+					if (!bMustBeSavedInDefaultDirectory || (LayoutFilePathAbsolute == DefaultDirectoryAbsolute))
+					{
+						// Save in the user layout folder
+						const FString& SourceFilePath = GEditorLayoutIni;
+						const FString& TargetFilePath = LayoutFilePath;
+						const bool bCleanLayoutNameAndDescriptionFieldsIfNoSameValues = !bMustBeSavedInDefaultDirectory;
+						const bool bShowSaveToast = true;
+						TrySaveLayoutOrWarnInternal(SourceFilePath, TargetFilePath, InWhatIsThis, bCleanLayoutNameAndDescriptionFieldsIfNoSameValues,
+							bShouldAskBeforeCleaningLayoutNameAndDescriptionFields, bShowSaveToast);
+					}
+					// If trying to write in a different folder (which is not allowed)
+					else
+					{
+						// Warn the user that the file will not be copied in there
+						const FText Title = LOCTEXT("SaveAsFailedMsg_Title", "Save As Failed");
+						FMessageDialog::Open(
+							EAppMsgType::Ok,
+							FText::Format(
+								LOCTEXT("SaveAsFailedMsg",
+									"In order to save the layout and allow Unreal to use it, you must save it in the predefined folder:\n{0}\n\nNevertheless, you tried to save it in:\n{1}\n\n"
+									"If you simply wish to export a copy of the current configuration in {1} (e.g., to later copy it into a different machine), you could use the \"Export Layout...\""
+									" functionality. However, Unreal would not be able to load it until you import it with \"Import Layout...\"."),
+								FText::FromString(DefaultDirectoryAbsolute), FText::FromString(LayoutFilePathAbsolute)),
+							&Title);
+					}
+				}
+			}
+		}
+	}
+
+	int32 GetNumberLayoutFiles(const FString& InLayoutsDirectory)
+	{
+		// Get Layout init files in desired directory
+		const TArray<FString> LayoutIniFileNames = GetIniFilesInFolderInternal(InLayoutsDirectory);
+		// Count how many layout files exist
+		int32 NumberLayoutFiles = 0;
+		for (const FString& LayoutIniFileName : LayoutIniFileNames)
+		{
+			const FString LayoutFilePath = FPaths::Combine(InLayoutsDirectory, LayoutIniFileName);
+			GConfig->UnloadFile(LayoutFilePath); // We must re-read it to avoid the Editor to use a previously cached name and description
+			if (FLayoutSaveRestore::IsValidConfig(LayoutFilePath))
+			{
+				++NumberLayoutFiles;
+			}
+		}
+		// Return result
+		return NumberLayoutFiles;
+	}
+};
+
 
 
 void FLayoutsMenuLoad::MakeLoadLayoutsMenu(UToolMenu* InToolMenu)
 {
 	// MakeLoadLayoutsMenu
 	const bool bDisplayDefaultLayouts = true;
-	MakeXLayoutsMenuInternal(InToolMenu, FMainFrameCommands::Get().MainFrameLayoutCommands.LoadLayoutCommands, FMainFrameCommands::Get().MainFrameLayoutCommands.LoadUserLayoutCommands, bDisplayDefaultLayouts);
+	FPrivateFLayoutsMenu::MakeXLayoutsMenuInternal(InToolMenu, bDisplayDefaultLayouts, FPrivateFLayoutsMenu::ELayoutsMenu::Load);
 
 	// Additional sections
 	if (GetDefault<UEditorStyleSettings>()->bEnableUserEditorLayoutManagement)
@@ -577,24 +797,24 @@ void FLayoutsMenuLoad::MakeLoadLayoutsMenu(UToolMenu* InToolMenu)
 
 		// Import...
 		{
-			Section.AddMenuEntry(FMainFrameCommands::Get().MainFrameLayoutCommands.ImportLayout);
+			Section.AddMenuEntry(FMainFrameCommands::Get().ImportLayout);
 		}
 	}
 }
 
 bool FLayoutsMenuLoad::CanLoadChooseLayout(const int32 InLayoutIndex)
 {
-	return !FLayoutsMenuBase::IsLayoutChecked(InLayoutIndex) && CanChooseLayoutWhenReadInternal(InLayoutIndex);
+	return !FLayoutsMenuBase::IsLayoutChecked(InLayoutIndex) && FPrivateFLayoutsMenu::CanChooseLayoutWhenReadInternal(InLayoutIndex);
 }
 bool FLayoutsMenuLoad::CanLoadChooseUserLayout(const int32 InLayoutIndex)
 {
-	return !FLayoutsMenuBase::IsUserLayoutChecked(InLayoutIndex) && CanChooseUserLayoutWhenReadInternal(InLayoutIndex);
+	return !FLayoutsMenuBase::IsUserLayoutChecked(InLayoutIndex) && FPrivateFLayoutsMenu::CanChooseUserLayoutWhenReadInternal(InLayoutIndex);
 }
 
 void FLayoutsMenuLoad::ReloadCurrentLayout()
 {
 	// If PIE, SIE, or any Asset Editors are opened, ask the user whether he wants to automatically close them and continue loading the layout
-	if (!CheckAskUserAndClosePIESIEAndAssetEditors(LOCTEXT("AreYouSureToLoadHeader", "Are you sure you want to continue loading the selected layout profile?")))
+	if (!FPrivateFLayoutsMenu::CheckAskUserAndClosePIESIEAndAssetEditors(LOCTEXT("AreYouSureToLoadHeader", "Are you sure you want to continue loading the selected layout profile?")))
 	{
 		return;
 	}
@@ -610,7 +830,7 @@ void FLayoutsMenuLoad::ReloadCurrentLayout()
 	const bool bShouldReplace = true;
 	const bool bEvenIfReadOnly = true;
 	const bool bCopyAttributes = false; // If true, it could e.g., copy the read-only flag of DefaultLayout.ini and make all the save/load stuff stop working
-	const FString& OriginalEditorLayoutIniFilePath = GetOriginalEditorLayoutIniFilePathInternal();
+	const FString& OriginalEditorLayoutIniFilePath = FPrivateFLayoutsMenu::GetOriginalEditorLayoutIniFilePathInternal();
 	IFileManager::Get().Copy(*OriginalEditorLayoutIniFilePath, *GEditorLayoutIni, bShouldReplace, bEvenIfReadOnly, bCopyAttributes);
 	GConfig->UnloadFile(OriginalEditorLayoutIniFilePath);
 	// Editor is reset on-the-fly
@@ -618,17 +838,17 @@ void FLayoutsMenuLoad::ReloadCurrentLayout()
 	EditorReinit();
 	FUnrealEdMisc::Get().AllowSavingLayoutOnClose(true);
 	// Save layout and create duplicated ini file (DuplicatedEditorLayoutIniFilePath)
-	SaveLayoutWithoutRemovingTempLayoutFiles();
+	FPrivateFLayoutsMenu::SaveLayoutWithoutRemovingTempLayoutFiles();
 	// If same file, remove temp files
 	const bool bCheckTempFileToo = false;
-	if (IsLayoutCheckedInternal(OriginalEditorLayoutIniFilePath, bCheckTempFileToo))
+	if (FPrivateFLayoutsMenu::IsLayoutCheckedInternal(OriginalEditorLayoutIniFilePath, bCheckTempFileToo))
 	{
-		RemoveTempEditorLayoutIniFilesInternal();
+		FPrivateFLayoutsMenu::RemoveTempEditorLayoutIniFilesInternal();
 	}
 	// Else, create DuplicatedEditorLayoutIniFilePath
 	else
 	{
-		const FString& DuplicatedEditorLayoutIniFilePath = GetDuplicatedEditorLayoutIniFilePathInternal();
+		const FString& DuplicatedEditorLayoutIniFilePath = FPrivateFLayoutsMenu::GetDuplicatedEditorLayoutIniFilePathInternal();
 		IFileManager::Get().Copy(*DuplicatedEditorLayoutIniFilePath, *GEditorLayoutIni, bShouldReplace, bEvenIfReadOnly, bCopyAttributes);
 		GConfig->UnloadFile(DuplicatedEditorLayoutIniFilePath);
 	}
@@ -641,7 +861,8 @@ void FLayoutsMenuLoad::LoadLayout(const FString& InLayoutPath)
 	const FString& TargetFilePath = GEditorLayoutIni;
 	const bool bCleanLayoutNameAndDescriptionFieldsIfNoSameValues = false;
 	const bool bShouldAskBeforeCleaningLayoutNameAndDescriptionFields = false;
-	const bool SucessfullySaved = TrySaveLayoutOrWarnInternal(SourceFilePath, TargetFilePath, LOCTEXT("LoadLayoutText", "loading the layout"), bCleanLayoutNameAndDescriptionFieldsIfNoSameValues, bShouldAskBeforeCleaningLayoutNameAndDescriptionFields);
+	const bool SucessfullySaved = FPrivateFLayoutsMenu::TrySaveLayoutOrWarnInternal(SourceFilePath, TargetFilePath, LOCTEXT("LoadLayoutText", "loading the layout"),
+		bCleanLayoutNameAndDescriptionFieldsIfNoSameValues, bShouldAskBeforeCleaningLayoutNameAndDescriptionFields);
 	// Reload current layout
 	if (SucessfullySaved)
 	{
@@ -685,7 +906,7 @@ void FLayoutsMenuLoad::ImportLayout()
 		if (bWereFilesSelected && LayoutFilePaths.Num() > 0)
 		{
 			// (Create if it does not exist and) Get UserLayoutsDirectory path
-			const FString UserLayoutsDirectory = CreateAndGetUserLayoutDirInternal();
+			const FString UserLayoutsDirectory = FPrivateFLayoutsMenu::CreateAndGetUserLayoutDirInternal();
 			// Iterate over selected layout ini files
 			FString FirstGoodLayoutFile = "";
 			const FText TrySaveLayoutOrWarnInternalText = LOCTEXT("ImportLayoutText", "importing the layout");
@@ -704,7 +925,8 @@ void FLayoutsMenuLoad::ImportLayout()
 					const FString& TargetFilePath = FPaths::Combine(FPaths::GetPath(UserLayoutsDirectory), FPaths::GetCleanFilename(LayoutFilePath));
 					const bool bCleanLayoutNameAndDescriptionFieldsIfNoSameValues = false;
 					const bool bShouldAskBeforeCleaningLayoutNameAndDescriptionFields = false;
-					TrySaveLayoutOrWarnInternal(SourceFilePath, TargetFilePath, TrySaveLayoutOrWarnInternalText, bCleanLayoutNameAndDescriptionFieldsIfNoSameValues, bShouldAskBeforeCleaningLayoutNameAndDescriptionFields);
+					FPrivateFLayoutsMenu::TrySaveLayoutOrWarnInternal(SourceFilePath, TargetFilePath, TrySaveLayoutOrWarnInternalText,
+						bCleanLayoutNameAndDescriptionFieldsIfNoSameValues, bShouldAskBeforeCleaningLayoutNameAndDescriptionFields);
 				}
 				// File is not a layout file, warn the user
 				else
@@ -717,7 +939,8 @@ void FLayoutsMenuLoad::ImportLayout()
 				}
 			}
 			// If PIE, SIE, or any Asset Editors are opened, ask the user whether he wants to automatically close them and continue loading the layout
-			if (!CheckAskUserAndClosePIESIEAndAssetEditors(LOCTEXT("LayoutImportClosePIEAndEditorAssetsHeader", "The layout(s) were successfully imported into the \"User Layouts\" section. Do you want to continue loading the selected layout profile?")))
+			if (!FPrivateFLayoutsMenu::CheckAskUserAndClosePIESIEAndAssetEditors(LOCTEXT("LayoutImportClosePIEAndEditorAssetsHeader",
+				"The layout(s) were successfully imported into the \"User Layouts\" section. Do you want to continue loading the selected layout profile?")))
 			{
 				return;
 			}
@@ -728,7 +951,8 @@ void FLayoutsMenuLoad::ImportLayout()
 				const FString& TargetFilePath = GEditorLayoutIni;
 				const bool bCleanLayoutNameAndDescriptionFieldsIfNoSameValues = false;
 				const bool bShouldAskBeforeCleaningLayoutNameAndDescriptionFields = false;
-				const bool SucessfullySaved = TrySaveLayoutOrWarnInternal(SourceFilePath, TargetFilePath, TrySaveLayoutOrWarnInternalText, bCleanLayoutNameAndDescriptionFieldsIfNoSameValues, bShouldAskBeforeCleaningLayoutNameAndDescriptionFields);
+				const bool SucessfullySaved = FPrivateFLayoutsMenu::TrySaveLayoutOrWarnInternal(SourceFilePath, TargetFilePath, TrySaveLayoutOrWarnInternalText,
+					bCleanLayoutNameAndDescriptionFieldsIfNoSameValues, bShouldAskBeforeCleaningLayoutNameAndDescriptionFields);
 				// Reload current layout
 				if (SucessfullySaved)
 				{
@@ -741,143 +965,13 @@ void FLayoutsMenuLoad::ImportLayout()
 
 
 
-FText GenerateLocalizedTextForFile(const FText& InText)
-{
-	// Proper FText to FString
-	FString StringSimulatingText;
-	FTextStringHelper::WriteToBuffer(StringSimulatingText, InText);
-	// Sanitize text (truncate if too big)
-	FString SanitazedTruncatedText = (StringSimulatingText.Len() < 100 ? StringSimulatingText : StringSimulatingText.Left(100));
-	FSaveLayoutDialogUtils::SanitizeText(SanitazedTruncatedText);
-	// Create full file path
-	if (StringSimulatingText.Len() < 10 || StringSimulatingText.Left(9) != TEXT("NSLOCTEXT"))
-	{
-		FString StringSimulatingTextRecreated =
-			// Namespace
-			TEXT("NSLOCTEXT(\"LayoutNamespace\", "
-			// Key
-			"\"") + SanitazedTruncatedText + TEXT("\", "
-			// Source string
-			"\"") + StringSimulatingText + TEXT("\")");
-		return FText::FromString(StringSimulatingTextRecreated);
-	}
-	else
-	{
-		return FText::FromString(StringSimulatingText);
-	}
-}
-
-void SaveExportLayoutCommon(const FString& InDefaultDirectory, const bool bMustBeSavedInDefaultDirectory, const FText& InWhatIsThis, const bool bShouldAskBeforeCleaningLayoutNameAndDescriptionFields)
-{
-	// Export/SaveAs the user-selected layout configuration files and load one of them
-	IDesktopPlatform* DesktopPlatform = FDesktopPlatformModule::Get();
-	if (DesktopPlatform)
-	{
-		bool bWereFilesSelected = false;
-		TArray<FString> LayoutFilePaths;
-		bool bWasDialogOpened = bMustBeSavedInDefaultDirectory;
-		// "Save Layout As..."
-		if (bMustBeSavedInDefaultDirectory)
-		{
-			TArray<FText> LayoutNames;
-			LayoutNames.Emplace(FLayoutSaveRestore::LoadSectionFromConfig(GEditorLayoutIni, "LayoutName"));
-			TArray<FText> LayoutDescriptions;
-			LayoutDescriptions.Emplace(FLayoutSaveRestore::LoadSectionFromConfig(GEditorLayoutIni, "LayoutDescription"));
-			// We want to avoid the duplication of the file/name/description fields, so we add "Copy of " at the beginning of the name and description fields
-			if (LayoutNames[0].ToString().Len() > 0)
-			{
-				LayoutNames[0] = FText::FromString(TEXT("Copy of ") + LayoutNames[0].ToString());
-			}
-			if (LayoutDescriptions[0].ToString().Len() > 0)
-			{
-				LayoutDescriptions[0] = FText::FromString(TEXT("Copy of ") + LayoutDescriptions[0].ToString());
-			}
-			// Create SWidget for saving the layout in its own SWindow and block the thread until it is finished
-			const TSharedRef<FSaveLayoutDialogParams> SaveLayoutDialogParams = MakeShared<FSaveLayoutDialogParams>(InDefaultDirectory, TEXT(".ini"), LayoutNames, LayoutDescriptions);
-			bWasDialogOpened = FSaveLayoutDialogUtils::CreateSaveLayoutAsDialogInStandaloneWindow(SaveLayoutDialogParams);
-			bWereFilesSelected = SaveLayoutDialogParams->bWereFilesSelected;
-			LayoutFilePaths = SaveLayoutDialogParams->LayoutFilePaths;
-			LayoutNames = SaveLayoutDialogParams->LayoutNames;
-			LayoutDescriptions = SaveLayoutDialogParams->LayoutDescriptions;
-
-			// Update GEditorLayoutIni file if LayoutNames or LayoutDescriptions were modified by the user
-			if (bWasDialogOpened && LayoutNames.Num() > 0 && LayoutDescriptions.Num() > 0 && (LayoutNames[0].ToString().Len() > 0 || LayoutDescriptions[0].ToString().Len() > 0))
-			{
-				checkf(LayoutNames.Num() == LayoutDescriptions.Num(), TEXT("There should be the same number of LayoutNames and LayoutDescriptions."));
-				for (int32 Index = 0; Index < LayoutNames.Num(); ++Index)
-				{
-					const FText& LayoutNameAsTextText = GenerateLocalizedTextForFile(LayoutNames[Index]);
-					const FText& LayoutDescriptionAsTextText = GenerateLocalizedTextForFile(LayoutDescriptions[Index]);
-					// Update fields
-					FLayoutSaveRestore::SaveSectionToConfig(GEditorLayoutIni, "LayoutName", LayoutNameAsTextText);
-					FLayoutSaveRestore::SaveSectionToConfig(GEditorLayoutIni, "LayoutDescription", LayoutDescriptionAsTextText);
-					// Flush file
-					const bool bRead = true;
-					GConfig->Flush(bRead, GEditorLayoutIni);
-				}
-			}
-		}
-		// "Export Layout..." (or "Save Layout As..." dialog could not be opened)
-		if (!bWasDialogOpened)
-		{
-			// Open the "save file" dialog so user can save his/her layout configuration file
-			const FString DefaultFile = "";
-			bWereFilesSelected = DesktopPlatform->SaveFileDialog(
-				FSlateApplication::Get().FindBestParentWindowHandleForDialogs(nullptr),
-				TEXT("Export a Layout Configuration File"),
-				InDefaultDirectory,
-				DefaultFile,
-				TEXT("Layout configuration files|*.ini|"),
-				EFileDialogFlags::None, //EFileDialogFlags::Multiple, // Allow/Avoid multiple file selection
-				LayoutFilePaths
-			);
-		}
-		// If file(s) selected, copy them into the user layouts directory and load one of them
-		if (bWereFilesSelected && LayoutFilePaths.Num() > 0)
-		{
-			// Iterate over selected layout ini files
-			FString FirstGoodLayoutFile = "";
-			for (const FString& LayoutFilePath : LayoutFilePaths)
-			{
-				// If writing in the right folder
-				const FString LayoutFilePathAbsolute = IFileManager::Get().ConvertToAbsolutePathForExternalAppForRead(*FPaths::GetPath(LayoutFilePath));
-				const FString DefaultDirectoryAbsolute = IFileManager::Get().ConvertToAbsolutePathForExternalAppForRead(*FPaths::GetPath(InDefaultDirectory));
-				if (!bMustBeSavedInDefaultDirectory || (LayoutFilePathAbsolute == DefaultDirectoryAbsolute))
-				{
-					// Save in the user layout folder
-					const FString& SourceFilePath = GEditorLayoutIni;
-					const FString& TargetFilePath = LayoutFilePath;
-					const bool bCleanLayoutNameAndDescriptionFieldsIfNoSameValues = !bMustBeSavedInDefaultDirectory;
-					const bool bShowSaveToast = true;
-					TrySaveLayoutOrWarnInternal(SourceFilePath, TargetFilePath, InWhatIsThis, bCleanLayoutNameAndDescriptionFieldsIfNoSameValues, bShouldAskBeforeCleaningLayoutNameAndDescriptionFields, bShowSaveToast);
-				}
-				// If trying to write in a different folder (which is not allowed)
-				else
-				{
-					// Warn the user that the file will not be copied in there
-					const FText Title = LOCTEXT("SaveAsFailedMsg_Title", "Save As Failed");
-					FMessageDialog::Open(
-						EAppMsgType::Ok,
-						FText::Format(
-							LOCTEXT("SaveAsFailedMsg",
-								"In order to save the layout and allow Unreal to use it, you must save it in the predefined folder:\n{0}\n\nNevertheless, you tried to save it in:\n{1}\n\n"
-								"If you simply wish to export a copy of the current configuration in {1} (e.g., to later copy it into a different machine), you could use the \"Export Layout...\""
-								" functionality. However, Unreal would not be able to load it until you import it with \"Import Layout...\"."),
-							FText::FromString(DefaultDirectoryAbsolute), FText::FromString(LayoutFilePathAbsolute)),
-						&Title);
-				}
-			}
-		}
-	}
-}
-
 void FLayoutsMenuSave::MakeSaveLayoutsMenu(UToolMenu* InToolMenu)
 {
 	if (GetDefault<UEditorStyleSettings>()->bEnableUserEditorLayoutManagement)
 	{
 		// MakeOverrideLayoutsMenu
 		const bool bDisplayDefaultLayouts = false;
-		MakeXLayoutsMenuInternal(InToolMenu, FMainFrameCommands::Get().MainFrameLayoutCommands.OverrideLayoutCommands, FMainFrameCommands::Get().MainFrameLayoutCommands.OverrideUserLayoutCommands, bDisplayDefaultLayouts);
+		FPrivateFLayoutsMenu::MakeXLayoutsMenuInternal(InToolMenu, bDisplayDefaultLayouts, FPrivateFLayoutsMenu::ELayoutsMenu::Save);
 
 		// Additional sections
 		{
@@ -891,12 +985,12 @@ void FLayoutsMenuSave::MakeSaveLayoutsMenu(UToolMenu* InToolMenu)
 
 			// Save as...
 			{
-				Section.AddMenuEntry(FMainFrameCommands::Get().MainFrameLayoutCommands.SaveLayoutAs);
+				Section.AddMenuEntry(FMainFrameCommands::Get().SaveLayoutAs);
 			}
 
 			// Export...
 			{
-				Section.AddMenuEntry(FMainFrameCommands::Get().MainFrameLayoutCommands.ExportLayout);
+				Section.AddMenuEntry(FMainFrameCommands::Get().ExportLayout);
 			}
 		}
 	}
@@ -904,11 +998,11 @@ void FLayoutsMenuSave::MakeSaveLayoutsMenu(UToolMenu* InToolMenu)
 
 bool FLayoutsMenuSave::CanSaveChooseLayout(const int32 InLayoutIndex)
 {
-	return !FLayoutsMenuBase::IsLayoutChecked(InLayoutIndex) && CanChooseLayoutWhenWriteInternal(InLayoutIndex);
+	return !FLayoutsMenuBase::IsLayoutChecked(InLayoutIndex) && FPrivateFLayoutsMenu::CanChooseLayoutWhenWriteInternal(InLayoutIndex);
 }
 bool FLayoutsMenuSave::CanSaveChooseUserLayout(const int32 InLayoutIndex)
 {
-	return !FLayoutsMenuBase::IsUserLayoutChecked(InLayoutIndex) && CanChooseUserLayoutWhenWriteInternal(InLayoutIndex);
+	return !FLayoutsMenuBase::IsUserLayoutChecked(InLayoutIndex) && FPrivateFLayoutsMenu::CanChooseUserLayoutWhenWriteInternal(InLayoutIndex);
 }
 
 void FLayoutsMenuSave::OverrideLayout(const int32 InLayoutIndex)
@@ -920,9 +1014,9 @@ void FLayoutsMenuSave::OverrideLayout(const int32 InLayoutIndex)
 void FLayoutsMenuSave::OverrideUserLayout(const int32 InLayoutIndex)
 {
 	// (Create if it does not exist and) Get UserLayoutsDirectory path
-	const FString UserLayoutsDirectory = CreateAndGetUserLayoutDirInternal();
+	const FString UserLayoutsDirectory = FPrivateFLayoutsMenu::CreateAndGetUserLayoutDirInternal();
 	// Get User Layout init files
-	const TArray<FString> UserLayoutIniFileNames = GetIniFilesInFolderInternal(UserLayoutsDirectory);
+	const TArray<FString> UserLayoutIniFileNames = FPrivateFLayoutsMenu::GetIniFilesInFolderInternal(UserLayoutsDirectory);
 	const FString DesiredUserLayoutFullPath = FPaths::Combine(FPaths::GetPath(UserLayoutsDirectory), UserLayoutIniFileNames[InLayoutIndex]);
 	// Are you sure you want to do this?
 	if (!FSaveLayoutDialogUtils::OverrideLayoutDialog(UserLayoutIniFileNames[InLayoutIndex]))
@@ -938,19 +1032,20 @@ void FLayoutsMenuSave::OverrideUserLayout(const int32 InLayoutIndex)
 	const bool bCleanLayoutNameAndDescriptionFieldsIfNoSameValues = true;
 	const bool bShouldAskBeforeCleaningLayoutNameAndDescriptionFields = false;
 	const bool bShowSaveToast = true;
-	TrySaveLayoutOrWarnInternal(SourceFilePath, TargetFilePath, LOCTEXT("OverrideLayoutText", "overriding the layout"), bCleanLayoutNameAndDescriptionFieldsIfNoSameValues, bShouldAskBeforeCleaningLayoutNameAndDescriptionFields, bShowSaveToast);
+	FPrivateFLayoutsMenu::TrySaveLayoutOrWarnInternal(SourceFilePath, TargetFilePath, LOCTEXT("OverrideLayoutText", "overriding the layout"),
+		bCleanLayoutNameAndDescriptionFieldsIfNoSameValues, bShouldAskBeforeCleaningLayoutNameAndDescriptionFields, bShowSaveToast);
 }
 
 void FLayoutsMenuSave::SaveLayout()
 {
 	// Save layout
-	SaveLayoutWithoutRemovingTempLayoutFiles();
+	FPrivateFLayoutsMenu::SaveLayoutWithoutRemovingTempLayoutFiles();
 	// Remove temporary Editor Layout ini files if the layout (thus also GEditorLayoutIni) changed
 	const bool bCheckTempFileToo = false;
-	const FString& DuplicatedEditorLayoutIniFilePath = GetDuplicatedEditorLayoutIniFilePathInternal();
-	if (!IsLayoutCheckedInternal(DuplicatedEditorLayoutIniFilePath, bCheckTempFileToo))
+	const FString& DuplicatedEditorLayoutIniFilePath = FPrivateFLayoutsMenu::GetDuplicatedEditorLayoutIniFilePathInternal();
+	if (!FPrivateFLayoutsMenu::IsLayoutCheckedInternal(DuplicatedEditorLayoutIniFilePath, bCheckTempFileToo))
 	{
-		RemoveTempEditorLayoutIniFilesInternal();
+		FPrivateFLayoutsMenu::RemoveTempEditorLayoutIniFilesInternal();
 	}
 }
 
@@ -959,10 +1054,11 @@ void FLayoutsMenuSave::SaveLayoutAs()
 	// Update GEditorLayoutIni file
 	SaveLayout();
 	// Copy GEditorLayoutIni into desired file
-	const FString DefaultDirectory = CreateAndGetUserLayoutDirInternal();
+	const FString DefaultDirectory = FPrivateFLayoutsMenu::CreateAndGetUserLayoutDirInternal();
 	const bool bMustBeSavedInDefaultDirectory = true;
 	const bool bShouldAskBeforeCleaningLayoutNameAndDescriptionFields = false;
-	SaveExportLayoutCommon(DefaultDirectory, bMustBeSavedInDefaultDirectory, LOCTEXT("SaveLayoutText", "saving the layout"), bShouldAskBeforeCleaningLayoutNameAndDescriptionFields);
+	FPrivateFLayoutsMenu::SaveExportLayoutCommon(DefaultDirectory, bMustBeSavedInDefaultDirectory, LOCTEXT("SaveLayoutText", "saving the layout"),
+		bShouldAskBeforeCleaningLayoutNameAndDescriptionFields);
 }
 
 void FLayoutsMenuSave::ExportLayout()
@@ -973,7 +1069,8 @@ void FLayoutsMenuSave::ExportLayout()
 	const FString DefaultDirectory = FPaths::ProjectContentDir();
 	const bool bMustBeSavedInDefaultDirectory = false;
 	const bool bShouldAskBeforeCleaningLayoutNameAndDescriptionFields = true;
-	SaveExportLayoutCommon(DefaultDirectory, bMustBeSavedInDefaultDirectory, LOCTEXT("ExportLayoutText", "exporting the layout"), bShouldAskBeforeCleaningLayoutNameAndDescriptionFields);
+	FPrivateFLayoutsMenu::SaveExportLayoutCommon(DefaultDirectory, bMustBeSavedInDefaultDirectory, LOCTEXT("ExportLayoutText", "exporting the layout"),
+		bShouldAskBeforeCleaningLayoutNameAndDescriptionFields);
 }
 
 
@@ -984,7 +1081,7 @@ void FLayoutsMenuRemove::MakeRemoveLayoutsMenu(UToolMenu* InToolMenu)
 	{
 		// MakeRemoveLayoutsMenu
 		const bool bDisplayDefaultLayouts = false;
-		MakeXLayoutsMenuInternal(InToolMenu, FMainFrameCommands::Get().MainFrameLayoutCommands.RemoveLayoutCommands, FMainFrameCommands::Get().MainFrameLayoutCommands.RemoveUserLayoutCommands, bDisplayDefaultLayouts);
+		FPrivateFLayoutsMenu::MakeXLayoutsMenuInternal(InToolMenu, bDisplayDefaultLayouts, FPrivateFLayoutsMenu::ELayoutsMenu::Remove);
 
 		// Additional sections
 		{
@@ -997,18 +1094,18 @@ void FLayoutsMenuRemove::MakeRemoveLayoutsMenu(UToolMenu* InToolMenu)
 			}
 
 			// Remove all
-			Section.AddMenuEntry(FMainFrameCommands::Get().MainFrameLayoutCommands.RemoveUserLayouts);
+			Section.AddMenuEntry(FMainFrameCommands::Get().RemoveUserLayouts);
 		}
 	}
 }
 
 bool FLayoutsMenuRemove::CanRemoveChooseLayout(const int32 InLayoutIndex)
 {
-	return CanChooseLayoutWhenWriteInternal(InLayoutIndex);
+	return FPrivateFLayoutsMenu::CanChooseLayoutWhenWriteInternal(InLayoutIndex);
 }
 bool FLayoutsMenuRemove::CanRemoveChooseUserLayout(const int32 InLayoutIndex)
 {
-	return CanChooseUserLayoutWhenWriteInternal(InLayoutIndex);
+	return FPrivateFLayoutsMenu::CanChooseUserLayoutWhenWriteInternal(InLayoutIndex);
 }
 
 void FLayoutsMenuRemove::RemoveLayout(const int32 InLayoutIndex)
@@ -1020,13 +1117,14 @@ void FLayoutsMenuRemove::RemoveLayout(const int32 InLayoutIndex)
 void FLayoutsMenuRemove::RemoveUserLayout(const int32 InLayoutIndex)
 {
 	// (Create if it does not exist and) Get UserLayoutsDirectory path
-	const FString UserLayoutsDirectory = CreateAndGetUserLayoutDirInternal();
+	const FString UserLayoutsDirectory = FPrivateFLayoutsMenu::CreateAndGetUserLayoutDirInternal();
 	// Get User Layout init files
-	const TArray<FString> UserLayoutIniFileNames = GetIniFilesInFolderInternal(UserLayoutsDirectory);
+	const TArray<FString> UserLayoutIniFileNames = FPrivateFLayoutsMenu::GetIniFilesInFolderInternal(UserLayoutsDirectory);
 	const FString DesiredUserLayoutFullPath = FPaths::Combine(FPaths::GetPath(UserLayoutsDirectory), UserLayoutIniFileNames[InLayoutIndex]);
 	// Are you sure you want to do this?
 	const FText TextFileNameToRemove = FText::FromString(FPaths::GetBaseFilename(UserLayoutIniFileNames[InLayoutIndex]));
-	const FText TextBody = FText::Format(LOCTEXT("ActionRemoveMsg", "Are you sure you want to permanently delete the layout profile \"{0}\"? This action cannot be undone."), TextFileNameToRemove);
+	const FText TextBody = FText::Format(LOCTEXT("ActionRemoveMsg",
+		"Are you sure you want to permanently delete the layout profile \"{0}\"? This action cannot be undone."), TextFileNameToRemove);
 	const FText TextTitle = FText::Format(LOCTEXT("RemoveUILayout_Title", "Remove UI Layout \"{0}\"?"), TextFileNameToRemove);
 	if (EAppReturnType::Ok != FMessageDialog::Open(EAppMsgType::OkCancel, TextBody, &TextTitle))
 	{
@@ -1036,43 +1134,25 @@ void FLayoutsMenuRemove::RemoveUserLayout(const int32 InLayoutIndex)
 	FPlatformFileManager::Get().GetPlatformFile().DeleteFile(*DesiredUserLayoutFullPath);
 }
 
-int32 GetNumberLayoutFiles(const FString& InLayoutsDirectory)
-{
-	// Get Layout init files in desired directory
-	const TArray<FString> LayoutIniFileNames = GetIniFilesInFolderInternal(InLayoutsDirectory);
-	// Count how many layout files exist
-	int32 NumberLayoutFiles = 0;
-	for (const FString& LayoutIniFileName : LayoutIniFileNames)
-	{
-		const FString LayoutFilePath = FPaths::Combine(InLayoutsDirectory, LayoutIniFileName);
-		GConfig->UnloadFile(LayoutFilePath); // We must re-read it to avoid the Editor to use a previously cached name and description
-		if (FLayoutSaveRestore::IsValidConfig(LayoutFilePath))
-		{
-			++NumberLayoutFiles;
-		}
-	}
-	// Return result
-	return NumberLayoutFiles;
-}
-
 void FLayoutsMenuRemove::RemoveUserLayouts()
 {
 	// (Create if it does not exist and) Get UserLayoutsDirectory path
-	const FString UserLayoutsDirectory = CreateAndGetUserLayoutDirInternal();
+	const FString UserLayoutsDirectory = FPrivateFLayoutsMenu::CreateAndGetUserLayoutDirInternal();
 	// Count how many layout files exist
-	const int32 NumberUserLayoutFiles = GetNumberLayoutFiles(UserLayoutsDirectory);
+	const int32 NumberUserLayoutFiles = FPrivateFLayoutsMenu::GetNumberLayoutFiles(UserLayoutsDirectory);
 	// If files to remove, warn user and remove them all
 	if (NumberUserLayoutFiles > 0)
 	{
 		// Are you sure you want to do this?
-		const FText TextBody = FText::Format(LOCTEXT("ActionRemoveAllUserLayoutMsg", "Are you sure you want to permanently remove {0} layout {0}|plural(one=profile,other=profiles)? This action cannot be undone."), NumberUserLayoutFiles);
+		const FText TextBody = FText::Format(LOCTEXT("ActionRemoveAllUserLayoutMsg",
+			"Are you sure you want to permanently remove {0} layout {0}|plural(one=profile,other=profiles)? This action cannot be undone."), NumberUserLayoutFiles);
 		const FText TextTitle = LOCTEXT("RemoveAllUserLayouts_Title", "Remove All User-Created Layouts?");
 		if (EAppReturnType::Ok != FMessageDialog::Open(EAppMsgType::OkCancel, TextBody, &TextTitle))
 		{
 			return;
 		}
 		// Get User Layout init files
-		const TArray<FString> UserLayoutIniFileNames = GetIniFilesInFolderInternal(UserLayoutsDirectory);
+		const TArray<FString> UserLayoutIniFileNames = FPrivateFLayoutsMenu::GetIniFilesInFolderInternal(UserLayoutsDirectory);
 		// If there are user Layout ini files, read them
 		for (const FString& UserLayoutIniFileName : UserLayoutIniFileNames)
 		{
@@ -1100,8 +1180,8 @@ void FLayoutsMenuRemove::RemoveUserLayouts()
 FString FLayoutsMenuBase::GetLayout(const int32 InLayoutIndex)
 {
 	// Get LayoutsDirectory path, layout init files, and desired layout path
-	const FString LayoutsDirectory = CreateAndGetDefaultLayoutDirInternal();
-	const TArray<FString> LayoutIniFileNames = GetIniFilesInFolderInternal(LayoutsDirectory);
+	const FString LayoutsDirectory = FPrivateFLayoutsMenu::CreateAndGetDefaultLayoutDirInternal();
+	const TArray<FString> LayoutIniFileNames = FPrivateFLayoutsMenu::GetIniFilesInFolderInternal(LayoutsDirectory);
 	const FString DesiredLayoutFullPath = FPaths::Combine(FPaths::GetPath(LayoutsDirectory), LayoutIniFileNames[InLayoutIndex]);
 	// Return full path
 	return DesiredLayoutFullPath;
@@ -1110,8 +1190,8 @@ FString FLayoutsMenuBase::GetLayout(const int32 InLayoutIndex)
 FString FLayoutsMenuBase::GetUserLayout(const int32 InLayoutIndex)
 {
 	// (Create if it does not exist and) Get UserLayoutsDirectory path, user layout init files, and desired user layout path
-	const FString UserLayoutsDirectory = CreateAndGetUserLayoutDirInternal();
-	const TArray<FString> UserLayoutIniFileNames = GetIniFilesInFolderInternal(UserLayoutsDirectory);
+	const FString UserLayoutsDirectory = FPrivateFLayoutsMenu::CreateAndGetUserLayoutDirInternal();
+	const TArray<FString> UserLayoutIniFileNames = FPrivateFLayoutsMenu::GetIniFilesInFolderInternal(UserLayoutsDirectory);
 	const FString DesiredUserLayoutFullPath = FPaths::Combine(FPaths::GetPath(UserLayoutsDirectory), UserLayoutIniFileNames[InLayoutIndex]);
 	// Return full path
 	return DesiredUserLayoutFullPath;
@@ -1120,9 +1200,9 @@ FString FLayoutsMenuBase::GetUserLayout(const int32 InLayoutIndex)
 bool FLayoutsMenuBase::IsThereUserLayouts()
 {
 	// (Create if it does not exist and) Get UserLayoutsDirectory path
-	const FString UserLayoutsDirectory = CreateAndGetUserLayoutDirInternal();
+	const FString UserLayoutsDirectory = FPrivateFLayoutsMenu::CreateAndGetUserLayoutDirInternal();
 	// At least 1 user layout file?
-	return GetNumberLayoutFiles(UserLayoutsDirectory) > 0;
+	return FPrivateFLayoutsMenu::GetNumberLayoutFiles(UserLayoutsDirectory) > 0;
 }
 
 bool FLayoutsMenuBase::IsLayoutChecked(const int32 InLayoutIndex)
@@ -1132,7 +1212,7 @@ bool FLayoutsMenuBase::IsLayoutChecked(const int32 InLayoutIndex)
 #else
 	// Check if the desired layout file matches the one currently loaded
 	const bool bCheckTempFileToo = true;
-	return IsLayoutCheckedInternal(GetLayout(InLayoutIndex), bCheckTempFileToo);
+	return FPrivateFLayoutsMenu::IsLayoutCheckedInternal(GetLayout(InLayoutIndex), bCheckTempFileToo);
 #endif
 }
 
@@ -1143,7 +1223,7 @@ bool FLayoutsMenuBase::IsUserLayoutChecked(const int32 InLayoutIndex)
 #else
 	// Check if the desired layout file matches the one currently loaded
 	const bool bCheckTempFileToo = true;
-	return IsLayoutCheckedInternal(GetUserLayout(InLayoutIndex), bCheckTempFileToo);
+	return FPrivateFLayoutsMenu::IsLayoutCheckedInternal(GetUserLayout(InLayoutIndex), bCheckTempFileToo);
 #endif
 }
 
