Add option to delay showing main frame at startup

diff --git a/Engine/Source/Editor/MainFrame/Private/MainFrameModule.cpp b/Engine/Source/Editor/MainFrame/Private/MainFrameModule.cpp
--- a/Engine/Source/Editor/MainFrame/Private/MainFrameModule.cpp
+++ b/Engine/Source/Editor/MainFrame/Private/MainFrameModule.cpp
@@ -238,8 +238,20 @@ void FMainFrameModule::CreateDefaultMainFrame( const bool bStartImmersive, const
 		// Initialize the main frame window
 		MainFrameHandler->OnMainFrameGenerated( MainTab, RootWindow );
 		
-		// Show the window!
-		MainFrameHandler->ShowMainFrameWindow( RootWindow, bStartImmersive, bStartPIE );
+		if (bDelayedShowMainFrame)
+		{
+			// Setup delegate to show main frame
+			DelayedShowMainFrameDelegate.BindLambda([this, RootWindow, bStartImmersive, bStartPIE]()
+			{
+				// Show the window!
+				MainFrameHandler->ShowMainFrameWindow(RootWindow, bStartImmersive, bStartPIE);
+			});
+		}
+		else
+		{
+			// Show the window!
+			MainFrameHandler->ShowMainFrameWindow(RootWindow, bStartImmersive, bStartPIE);
+		}
 		
 		MRUFavoritesList = new FMainMRUFavoritesList;
 		MRUFavoritesList->ReadFromINI();
@@ -456,6 +468,8 @@ void FMainFrameModule::SetApplicationTitleOverride(const FText& NewOverriddenApp
 
 void FMainFrameModule::StartupModule( )
 {
+	bDelayedShowMainFrame = false;
+
 	MRUFavoritesList = NULL;
 
 	ensureMsgf(!IsRunningGame(), TEXT("The MainFrame module should only be loaded when running the editor.  Code that extends the editor, adds menu items, etc... should not run when running in -game mode or in a non-WITH_EDITOR build"));
@@ -496,6 +510,8 @@ void FMainFrameModule::StartupModule( )
 
 void FMainFrameModule::ShutdownModule( )
 {
+	ClearDelayedShowMainFrameDelegate();
+
 	// Destroy the main frame window
 	TSharedPtr< SWindow > ParentWindow( GetParentWindow() );
 	if( ParentWindow.IsValid() )
diff --git a/Engine/Source/Editor/MainFrame/Private/MainFrameModule.h b/Engine/Source/Editor/MainFrame/Private/MainFrameModule.h
--- a/Engine/Source/Editor/MainFrame/Private/MainFrameModule.h
+++ b/Engine/Source/Editor/MainFrame/Private/MainFrameModule.h
@@ -55,6 +55,8 @@ public:
 
 	virtual void RequestCloseEditor( ) override
 	{
+		ClearDelayedShowMainFrameDelegate();
+
 		if ( MainFrameHandler->CanCloseEditor() )
 		{
 			MainFrameHandler->ShutDownEditor();
@@ -110,6 +112,22 @@ public:
 		return MainFrameSDKNotInstalled.Broadcast(PlatformName, DocLink);
 	}
 
+	virtual void EnableDelayedShowMainFrame() override
+	{
+		bDelayedShowMainFrame = true;
+	}
+
+	virtual void ShowDelayedMainFrame() override
+	{
+		bDelayedShowMainFrame = false;
+
+		if (DelayedShowMainFrameDelegate.IsBound())
+		{
+			DelayedShowMainFrameDelegate.Execute();
+			ClearDelayedShowMainFrameDelegate();
+		}
+	}
+
 public:
 
 	// IModuleInterface interface
@@ -161,6 +179,12 @@ private:
 	// Handles launching code accessor
 	void HandleCodeAccessorLaunching( );
 
+	// Reset delegate
+	void ClearDelayedShowMainFrameDelegate()
+	{
+		DelayedShowMainFrameDelegate.Unbind();
+	}
+
 private:
 
 	// Weak pointer to the level editor's compile notification item.
@@ -193,6 +217,12 @@ private:
 	// Weak pointer to the code accessor's notification item.
 	TWeakPtr<class SNotificationItem> CodeAccessorNotificationPtr;
 
+	// Delegate that holds a delayed call to ShowMainFrameWindow
+	FSimpleDelegate DelayedShowMainFrameDelegate;
+
+	// Allow delaying when to show main frame's window
+	bool bDelayedShowMainFrame;
+
 	// Sounds used for compilation.
 	class USoundBase* CompileStartSound;
 	class USoundBase* CompileSuccessSound;
diff --git a/Engine/Source/Editor/MainFrame/Public/Interfaces/IMainFrameModule.h b/Engine/Source/Editor/MainFrame/Public/Interfaces/IMainFrameModule.h
--- a/Engine/Source/Editor/MainFrame/Public/Interfaces/IMainFrameModule.h
+++ b/Engine/Source/Editor/MainFrame/Public/Interfaces/IMainFrameModule.h
@@ -178,6 +178,16 @@ public:
 	virtual FMainFrameSDKNotInstalled& OnMainFrameSDKNotInstalled( ) = 0;
 	virtual void BroadcastMainFrameSDKNotInstalled(const FString& PlatformName, const FString& DocLink) = 0;
 
+	/**
+	 * Enable external control of when main frame is shown
+	 */
+	virtual void EnableDelayedShowMainFrame() = 0;
+
+	/**
+	 * Show main frame now if it was delayed and not shown yet
+	 */
+	virtual void ShowDelayedMainFrame() = 0;
+
 public:
 
 	/**
