Fix packaging triggering a code build even if we have a valid promoted binary.

#rnx

diff --git a/Engine/Source/Editor/MainFrame/Private/Frame/MainFrameActions.cpp b/Engine/Source/Editor/MainFrame/Private/Frame/MainFrameActions.cpp
--- a/Engine/Source/Editor/MainFrame/Private/Frame/MainFrameActions.cpp
+++ b/Engine/Source/Editor/MainFrame/Private/Frame/MainFrameActions.cpp
@@ -46,6 +46,7 @@
 #include "Dialogs/SOutputLogDialog.h"
 #include "IUATHelperModule.h"
 #include "Menus/LayoutsMenu.h"
+#include "TargetReceipt.h"
 
 #include "Settings/EditorSettings.h"
 #include "AnalyticsEventAttribute.h"
@@ -726,6 +727,9 @@ void FMainFrameActionCallbacks::PackageProject( const FName InPlatformInfoName )
 		OptionalParams += *PlatformInfo->TargetPlatformName.ToString();
 	}
 
+	// Get the target to build
+	const FTargetInfo* Target = PackagingSettings->GetBuildTargetInfo();
+
 	// Only build if the user elects to do so
 	bool bBuild = false;
 	if(PackagingSettings->Build == EProjectPackagingBuild::Always)
@@ -738,8 +742,53 @@ void FMainFrameActionCallbacks::PackageProject( const FName InPlatformInfoName )
 	}
 	else if(PackagingSettings->Build == EProjectPackagingBuild::IfProjectHasCode)
 	{
-		FText Reason;
-		bBuild = bProjectHasCode || !FApp::GetEngineIsPromotedBuild() || (Platform == nullptr || Platform->RequiresTempTarget(bProjectHasCode, ConfigurationInfo.Configuration, bAssetNativizationEnabled, Reason));
+		bBuild = true;
+		if (FApp::GetEngineIsPromotedBuild())
+		{
+			FString BaseDir;
+
+			// Get the target name
+			FString TargetName;
+			if (Target == nullptr)
+			{
+				TargetName = TEXT("UE4Game");
+			}
+			else
+			{
+				TargetName = Target->Name;
+			}
+
+			// Get the directory containing the receipt for this target, depending on whether the project needs to be built or not
+			FString ProjectDir = FPaths::GetPath(FPaths::GetProjectFilePath());
+			if (Target != nullptr && FPaths::IsUnderDirectory(Target->Path, ProjectDir))
+			{
+				UE_LOG(LogMainFrame, Log, TEXT("Selected target: %s"), *Target->Name);
+				BaseDir = ProjectDir;
+			}
+			else
+			{
+				FText Reason;
+				if (Platform->RequiresTempTarget(bProjectHasCode, ConfigurationInfo.Configuration, false, Reason))
+				{
+					UE_LOG(LogMainFrame, Log, TEXT("Project requires temp target (%s)"), *Reason.ToString());
+					BaseDir = ProjectDir;
+				}
+				else
+				{
+					UE_LOG(LogMainFrame, Log, TEXT("Project does not require temp target"));
+					BaseDir = FPaths::EngineDir();
+				}
+			}
+
+			// Check if the receipt is for a matching promoted target
+			FString PlatformName = Platform->GetPlatformInfo().UBTTargetId.ToString();
+
+			extern UNREALED_API bool HasPromotedTarget(const TCHAR* BaseDir, const TCHAR* TargetName, const TCHAR* Platform, EBuildConfiguration Configuration, const TCHAR* Architecture);
+			if (HasPromotedTarget(*BaseDir, *TargetName, *PlatformName, ConfigurationInfo.Configuration, nullptr))
+			{
+				bBuild = false;
+			}
+		}
 	}
 	else if(PackagingSettings->Build == EProjectPackagingBuild::IfEditorWasBuiltLocally)
 	{
@@ -767,7 +816,6 @@ void FMainFrameActionCallbacks::PackageProject( const FName InPlatformInfoName )
 		OptionalParams += FString::Printf(TEXT(" -NumCookersToSpawn=%d"), NumCookers); 
 	}
 
-	const FTargetInfo* Target = PackagingSettings->GetBuildTargetInfo();
 	if (Target == nullptr)
 	{
 		OptionalParams += FString::Printf(TEXT(" -clientconfig=%s"), LexToString(ConfigurationInfo.Configuration));
diff --git a/Engine/Source/Editor/UnrealEd/Private/PlayLevelViaLauncher.cpp b/Engine/Source/Editor/UnrealEd/Private/PlayLevelViaLauncher.cpp
--- a/Engine/Source/Editor/UnrealEd/Private/PlayLevelViaLauncher.cpp
+++ b/Engine/Source/Editor/UnrealEd/Private/PlayLevelViaLauncher.cpp
@@ -52,7 +52,7 @@ static void HandleOutputReceived(const FString& InMessage)
 	}
 }
 
-bool HasPromotedTarget(const TCHAR* BaseDir, const TCHAR* TargetName, const TCHAR* Platform, EBuildConfiguration Configuration, const TCHAR* Architecture)
+UNREALED_API bool HasPromotedTarget(const TCHAR* BaseDir, const TCHAR* TargetName, const TCHAR* Platform, EBuildConfiguration Configuration, const TCHAR* Architecture)
 {
 	// Get the path to the receipt, and check it exists
 	FString ReceiptPath = FTargetReceipt::GetDefaultPath(BaseDir, TargetName, Platform, Configuration, Architecture);
