- Fixed cis error
#rnx

diff --git a/Engine/Source/Editor/MainFrame/Private/Menus/PackageProjectMenu.h b/Engine/Source/Editor/MainFrame/Private/Menus/PackageProjectMenu.h
--- a/Engine/Source/Editor/MainFrame/Private/Menus/PackageProjectMenu.h
+++ b/Engine/Source/Editor/MainFrame/Private/Menus/PackageProjectMenu.h
@@ -16,6 +16,7 @@
 #include "Interfaces/IProjectManager.h"
 #include "InstalledPlatformInfo.h"
 #include "Misc/DataDrivenPlatformInfoRegistry.h"
+#include "Misc/ConfigCacheIni.h"
 #include "DesktopPlatformModule.h"
 
 #define LOCTEXT_NAMESPACE "FPackageProjectMenu"
