- Fixed the Packaging PlatformName for Windows, and some mild future proofing for allowing Client target types (which is not exposed)

diff --git a/Engine/Source/Editor/MainFrame/Private/Frame/MainFrameActions.cpp b/Engine/Source/Editor/MainFrame/Private/Frame/MainFrameActions.cpp
--- a/Engine/Source/Editor/MainFrame/Private/Frame/MainFrameActions.cpp
+++ b/Engine/Source/Editor/MainFrame/Private/Frame/MainFrameActions.cpp
@@ -705,17 +705,17 @@ void FMainFrameActionCallbacks::PackageProject( const FName InPlatformInfoName )
 		OptionalParams += TEXT(" -manifests");
 	}
 
+	OptionalParams += TEXT(" -targetplatform=");
+	OptionalParams += PlatformInfo->DataDrivenPlatformInfo->UBTPlatformString;
+
 	// Append any extra UAT flags specified for this platform flavor
 	if (!PlatformInfo->UATCommandLine.IsEmpty())
 	{
 		OptionalParams += TEXT(" ");
 		OptionalParams += PlatformInfo->UATCommandLine;
 	}
-	else
-	{
-		OptionalParams += TEXT(" -targetplatform=");
-		OptionalParams += *PlatformInfo->Name.ToString();
-	}
+
+
 
 	// Get the target to build
 	const FTargetInfo* Target = PackagingSettings->GetBuildTargetInfo();
@@ -819,7 +819,7 @@ void FMainFrameActionCallbacks::PackageProject( const FName InPlatformInfoName )
 		OptionalParams += FString::Printf(TEXT(" -target=%s -clientconfig=%s"), *Target->Name, LexToString(ConfigurationInfo.Configuration));
 	}
 
-	FString TurnkeyCommandline = FString::Printf(TEXT("Turnkey -command=VerifySdk -platform=%s -UpdateIfNeeded -EditorIO"), *PlatformInfo->IniPlatformName.ToString());
+	FString TurnkeyCommandline = FString::Printf(TEXT("Turnkey -command=VerifySdk -platform=%s -UpdateIfNeeded -EditorIO"), *PlatformInfo->DataDrivenPlatformInfo->UBTPlatformString);
 
 	FString ProjectPath = FPaths::IsProjectFilePathSet() ? FPaths::ConvertRelativePathToFull(FPaths::GetProjectFilePath()) : FPaths::RootDir() / FApp::GetProjectName() / FApp::GetProjectName() + TEXT(".uproject");
 	FString CommandLine = FString::Printf(TEXT("-ScriptsForProject=\"%s\" %s BuildCookRun %s%s -nop4 -project=\"%s\" -cook -stage -archive -archivedirectory=\"%s\" -package -ue4exe=\"%s\" %s -utf8output"),
diff --git a/Engine/Source/Editor/UnrealEd/Private/Kismet2/DebuggerCommands.cpp b/Engine/Source/Editor/UnrealEd/Private/Kismet2/DebuggerCommands.cpp
--- a/Engine/Source/Editor/UnrealEd/Private/Kismet2/DebuggerCommands.cpp
+++ b/Engine/Source/Editor/UnrealEd/Private/Kismet2/DebuggerCommands.cpp
@@ -2786,15 +2786,19 @@ void FInternalPlayWorldCommandCallbacks::HandleLaunchOnDeviceActionExecute(FStri
 		FString TargetPlatformName = DeviceId.Left(Index);
 		const PlatformInfo::FTargetPlatformInfo* const PlatformInfo = PlatformInfo::FindPlatformInfo(FName(*TargetPlatformName));
 		FString IniPlatformName = PlatformInfo->IniPlatformName.ToString();
+		FString UBTPlatformName = PlatformInfo->DataDrivenPlatformInfo->UBTPlatformString;
 
-		FString CommandLine = FString::Printf(TEXT("Turnkey -command=VerifySdk -UpdateIfNeeded -platform=%s -EditorIO -noturnkeyvariables -device=%s"), *IniPlatformName, *DeviceName);
+		FString CommandLine = FString::Printf(TEXT("Turnkey -command=VerifySdk -UpdateIfNeeded -platform=%s -EditorIO -noturnkeyvariables -device=%s"), *UBTPlatformName, *DeviceName);
 		FText TaskName = LOCTEXT("VerifyingSDK", "Verifying SDK and Device");
 
 		IUATHelperModule::Get().CreateUatTask(CommandLine, FText::FromString(IniPlatformName), TaskName, TaskName, FEditorStyle::GetBrush(TEXT("MainFrame.PackageProject")),
 			[DeviceId, DeviceName](FString, double)
 		{
-			PrepareLaunchOn(DeviceId, DeviceName);
-			LaunchOnDevice(DeviceId, DeviceName, false);
+			AsyncTask(ENamedThreads::GameThread, [DeviceId, DeviceName]()
+			{
+				PrepareLaunchOn(DeviceId, DeviceName);
+				LaunchOnDevice(DeviceId, DeviceName, false);
+			});
 		});
 	}
 	else
