Fix up save layout dialog. Updated styles and cut off buttons

diff --git a/Engine/Source/Editor/ContentBrowser/Private/SAssetDialog.cpp b/Engine/Source/Editor/ContentBrowser/Private/SAssetDialog.cpp
--- a/Engine/Source/Editor/ContentBrowser/Private/SAssetDialog.cpp
+++ b/Engine/Source/Editor/ContentBrowser/Private/SAssetDialog.cpp
@@ -181,7 +181,6 @@ void SAssetDialog::Construct(const FArguments& InArgs, const FSharedAssetDialogC
 					SNew(STextBlock)
 					.Text( this, &SAssetDialog::GetNameErrorLabelText )
 					.ToolTipText(this, &SAssetDialog::GetNameErrorLabelText)
-					.TextStyle( FEditorStyle::Get(), "AssetDialog.ErrorLabelFont" )
 				]
 			]
 		];
diff --git a/Engine/Source/Editor/EditorStyle/Private/StarshipStyle.cpp b/Engine/Source/Editor/EditorStyle/Private/StarshipStyle.cpp
--- a/Engine/Source/Editor/EditorStyle/Private/StarshipStyle.cpp
+++ b/Engine/Source/Editor/EditorStyle/Private/StarshipStyle.cpp
@@ -7570,8 +7570,7 @@ void FStarshipEditorStyle::FStyle::SetupAutomationStyles()
 
 	// Asset Dialog
 	{
-		Set("AssetDialog.ErrorLabelBorder", new FSlateColorBrush(FLinearColor(0.2, 0, 0)));
-		Set("AssetDialog.ErrorLabelFont", FTextBlockStyle(NormalText).SetFont(DEFAULT_FONT("Regular", 10)));
+		Set("AssetDialog.ErrorLabelBorder", new FSlateColorBrush(FStyleColors::AccentRed));
 	}
 }
 
diff --git a/Engine/Source/Editor/MainFrame/Private/Menus/SaveLayoutDialog.cpp b/Engine/Source/Editor/MainFrame/Private/Menus/SaveLayoutDialog.cpp
--- a/Engine/Source/Editor/MainFrame/Private/Menus/SaveLayoutDialog.cpp
+++ b/Engine/Source/Editor/MainFrame/Private/Menus/SaveLayoutDialog.cpp
@@ -130,11 +130,11 @@ bool FSaveLayoutDialogUtils::CreateSaveLayoutAsDialogInStandaloneWindow(const TS
 	TSharedRef<SSaveLayoutDialog> SaveLayoutDialog = SNew(SSaveLayoutDialog, InSaveLayoutDialogParams);
 
 	// Create SWindow that contains SSaveLayoutDialog
-	const FVector2D WindowSize(720.f, 130.f);
 	TSharedRef<SWindow> DialogWindow =
 		SNew(SWindow)
 		.Title(LOCTEXT("GenericAssetDialogWindowHeader", "Save Layout As"))
-		.ClientSize(WindowSize);
+		.SizingRule(ESizingRule::Autosized);
+
 	DialogWindow->SetContent(SaveLayoutDialog);
 
 	// Launch SSaveLayoutDialog and block thread until user finishes with it
@@ -162,7 +162,6 @@ SSaveLayoutDialog::~SSaveLayoutDialog()
 {
 }
 
-BEGIN_SLATE_FUNCTION_BUILD_OPTIMIZATION
 void SSaveLayoutDialog::Construct(const FArguments& InArgs, const TSharedRef<FSaveLayoutDialogParams>& InSaveLayoutDialogParams)
 {
 	SaveLayoutDialogParams = InSaveLayoutDialogParams;
@@ -186,132 +185,120 @@ void SSaveLayoutDialog::Construct(const FArguments& InArgs, const TSharedRef<FSa
 	// UI part
 	// The root SWidget in this dialog
 	TSharedRef<SVerticalBox> MainVerticalBox = SNew(SVerticalBox);
-	TSharedRef<SVerticalBox> LabelsBox = SNew(SVerticalBox);
 	TSharedRef<SVerticalBox> ContentBox = SNew(SVerticalBox);
 
-	// Layout Name
-	LabelsBox->AddSlot()
-		.FillHeight(1)
-		.VAlign(VAlign_Center)
-		.Padding(0, 3.625f, 0, 3.625f)
-		[
-			SNew(STextBlock).Text(LOCTEXT("LayoutNameBoxLabel", "Name:"))
-		];
-	ContentBox->AddSlot()
-		.FillHeight(1)
-		.VAlign(VAlign_Center)
-		.Padding(0, 2.f, 0, 2.f)
-		[
-			SNew(SEditableTextBox)
-			.Text(CurrentlyEnteredLayoutName)
-			.OnTextCommitted(this, &SSaveLayoutDialog::OnLayoutNameTextCommited)
-			.OnTextChanged(this, &SSaveLayoutDialog::OnLayoutNameTextCommited, ETextCommit::Default)
-			.SelectAllTextWhenFocused(true)
-		];
-
-	// Layout Description
-	LabelsBox->AddSlot()
-		.FillHeight(1)
-		.VAlign(VAlign_Center)
-		.Padding(0, 3.625f, 0, 3.625f)
-		[
-			SNew(STextBlock).Text(LOCTEXT("LayoutDescriptionBoxLabel", "Description (Optional):"))
-		];
-	ContentBox->AddSlot()
-		.FillHeight(1)
-		.VAlign(VAlign_Center)
-		.Padding(0, 2.f, 0, 2.f)
-		[
-			SNew(SEditableTextBox)
-			.Text(CurrentlyEnteredLayoutDescription)
-			.OnTextCommitted(this, &SSaveLayoutDialog::OnLayoutDescriptionTextCommited)
-			.OnTextChanged(this, &SSaveLayoutDialog::OnLayoutDescriptionTextCommited, ETextCommit::Default)
-			.SelectAllTextWhenFocused(true)
-		];
-
-	// Buttons and asset name
-	TSharedRef<SHorizontalBox> LabelAndContentBox = SNew(SHorizontalBox)
-		+ SHorizontalBox::Slot()
-		.AutoWidth()
-		.HAlign(HAlign_Right)
-		.VAlign(VAlign_Bottom)
-		.Padding(30, 30, 4, 3)
-		[
-			LabelsBox
-		]
-		+ SHorizontalBox::Slot()
-		.FillWidth(1)
-		.VAlign(VAlign_Bottom)
-		.Padding(4, 3, 30, 4)
-		[
-			ContentBox
-		];
-	MainVerticalBox->AddSlot()
-		.AutoHeight()
-		.HAlign(HAlign_Fill)
-		.Padding(0)
-		[
-			LabelAndContentBox
-		];
-
-	// Name Error label
-	MainVerticalBox->AddSlot()
-	.AutoHeight()
+	// Add to ChildSlot
+	ChildSlot
 	[
-		// Constant height, whether the label is visible or not
-		SNew(SBox).HeightOverride(20)
+		SNew(SBorder)
+		.BorderImage(FAppStyle::Get().GetBrush("Brushes.Background"))
 		[
-			SNew(SBorder)
-			.Visibility( this, &SSaveLayoutDialog::GetNameErrorLabelVisibility )
-			.BorderImage( FEditorStyle::GetBrush("AssetDialog.ErrorLabelBorder") )
-			.Content()
+			SNew(SBox)
+			.WidthOverride(600.f)
 			[
-				SNew(STextBlock)
-				.Text( this, &SSaveLayoutDialog::GetNameErrorLabelText )
-				.ToolTipText(this, &SSaveLayoutDialog::GetNameErrorLabelText)
-				.TextStyle( FEditorStyle::Get(), "AssetDialog.ErrorLabelFont" )
+				// Layout Name
+				SNew(SVerticalBox)
+				+ SVerticalBox::Slot()
+				.AutoHeight()
+				.Padding(5.0f,2.0f)
+				[
+					SNew(SHorizontalBox)
+					+SHorizontalBox::Slot()
+					.FillWidth(0.5f)
+					.HAlign(HAlign_Right)
+					.VAlign(VAlign_Center)
+					.Padding(5.0f, 2.0f)
+					[	
+						SNew(STextBlock)
+						.Text(LOCTEXT("LayoutNameBoxLabel", "Name"))
+					]
+					+SHorizontalBox::Slot()
+					.FillWidth(2.0f)
+					.VAlign(VAlign_Center)
+					.Padding(5.0f, 2.0f)
+					[
+						SNew(SEditableTextBox)
+						.Text(CurrentlyEnteredLayoutName)
+						.OnTextCommitted(this, &SSaveLayoutDialog::OnLayoutNameTextCommited)
+						.OnTextChanged(this, &SSaveLayoutDialog::OnLayoutNameTextCommited, ETextCommit::Default)
+						.SelectAllTextWhenFocused(true)
+					]
+				]
+				+ SVerticalBox::Slot()
+				.Padding(5.0f, 2.0f)
+				.AutoHeight()
+				[
+					SNew(SHorizontalBox)
+					+ SHorizontalBox::Slot()
+					.FillWidth(.5f)
+					.HAlign(HAlign_Right)
+					.VAlign(VAlign_Center)
+					.Padding(5.0f, 2.0f)
+					[
+						SNew(STextBlock)
+						.Text(LOCTEXT("LayoutDescriptionBoxLabel", "Description (Optional)"))
+					]
+					+ SHorizontalBox::Slot()
+					.VAlign(VAlign_Center)
+					.FillWidth(2.0f)
+					.Padding(5.0f, 2.0f)
+					[
+						SNew(SEditableTextBox)
+						.Text(CurrentlyEnteredLayoutDescription)
+						.OnTextCommitted(this, &SSaveLayoutDialog::OnLayoutDescriptionTextCommited)
+						.OnTextChanged(this, &SSaveLayoutDialog::OnLayoutDescriptionTextCommited, ETextCommit::Default)
+						.SelectAllTextWhenFocused(true)
+					]
+				]
+				+ SVerticalBox::Slot()
+				.AutoHeight()
+				[
+					// Constant height, whether the label is visible or not
+					SNew(SBox)
+					.HeightOverride(20)
+					[
+						SNew(SBorder)
+						.Visibility( this, &SSaveLayoutDialog::GetNameErrorLabelVisibility )
+						.BorderImage( FEditorStyle::GetBrush("AssetDialog.ErrorLabelBorder") )
+						.Content()
+						[
+							SNew(STextBlock)
+							.Text( this, &SSaveLayoutDialog::GetNameErrorLabelText )
+							.ToolTipText(this, &SSaveLayoutDialog::GetNameErrorLabelText)
+						]
+					]
+				]
+				+ SVerticalBox::Slot()
+				.HAlign(HAlign_Right)
+				.VAlign(VAlign_Bottom)
+				.Padding(0, 0, 5, 5)
+				[
+					SNew(SHorizontalBox)
+					+ SHorizontalBox::Slot()
+					.AutoWidth()
+					.VAlign(VAlign_Bottom)
+					.Padding(4, 3)
+					[
+						SNew(SButton)
+						.Text(LOCTEXT("SaveLayoutDialogSaveButton", "Save"))
+						.ButtonStyle(FAppStyle::Get(), "PrimaryButton")
+						.IsEnabled(this, &SSaveLayoutDialog::IsConfirmButtonEnabled)
+						.OnClicked(this, &SSaveLayoutDialog::OnConfirmClicked)
+					]
+					+ SHorizontalBox::Slot()
+					.AutoWidth()
+					.VAlign(VAlign_Bottom)
+					.Padding(4, 3)
+					[
+						SNew(SButton)
+						.Text(LOCTEXT("SaveLayoutDialogCancelButton", "Cancel"))
+						.OnClicked(this, &SSaveLayoutDialog::OnCancelClicked)
+					]
+				]
 			]
 		]
 	];
-
-	// Buttons
-	TSharedRef<SHorizontalBox> ButtonBox = SNew(SHorizontalBox)
-		+ SHorizontalBox::Slot()
-		.AutoWidth()
-		.VAlign(VAlign_Bottom)
-		.Padding(4, 3)
-		[
-			SNew(SButton)
-			.Text(LOCTEXT("SaveLayoutDialogSaveButton", "Save"))
-			.ContentPadding(FMargin(8, 2, 8, 2))
-			.IsEnabled(this, &SSaveLayoutDialog::IsConfirmButtonEnabled)
-			.OnClicked(this, &SSaveLayoutDialog::OnConfirmClicked)
-		]
-		+ SHorizontalBox::Slot()
-		.AutoWidth()
-		.VAlign(VAlign_Bottom)
-		.Padding(4, 3)
-		[
-			SNew(SButton)
-			.ContentPadding(FMargin(8, 2, 8, 2))
-			.Text(LOCTEXT("SaveLayoutDialogCancelButton", "Cancel"))
-			.OnClicked(this, &SSaveLayoutDialog::OnCancelClicked)
-		];
-	MainVerticalBox->AddSlot()
-		.HAlign(HAlign_Right)
-		.VAlign(VAlign_Bottom)
-		.Padding(0, 0, 5, 5)
-		[
-			ButtonBox
-		];
-
-	// Add to ChildSlot
-	ChildSlot
-	[
-		MainVerticalBox
-	];
 }
-END_SLATE_FUNCTION_BUILD_OPTIMIZATION
 
 FText SSaveLayoutDialog::GetFileNameText() const
 {
