Fix cut off text in the save layout dialog

diff --git a/Engine/Source/Editor/MainFrame/Private/Menus/SaveLayoutDialog.cpp b/Engine/Source/Editor/MainFrame/Private/Menus/SaveLayoutDialog.cpp
--- a/Engine/Source/Editor/MainFrame/Private/Menus/SaveLayoutDialog.cpp
+++ b/Engine/Source/Editor/MainFrame/Private/Menus/SaveLayoutDialog.cpp
@@ -204,7 +204,7 @@ void SSaveLayoutDialog::Construct(const FArguments& InArgs, const TSharedRef<FSa
 				[
 					SNew(SHorizontalBox)
 					+SHorizontalBox::Slot()
-					.FillWidth(0.5f)
+					.FillWidth(0.6f)
 					.HAlign(HAlign_Right)
 					.VAlign(VAlign_Center)
 					.Padding(5.0f, 2.0f)
@@ -230,7 +230,7 @@ void SSaveLayoutDialog::Construct(const FArguments& InArgs, const TSharedRef<FSa
 				[
 					SNew(SHorizontalBox)
 					+ SHorizontalBox::Slot()
-					.FillWidth(.5f)
+					.FillWidth(.6f)
 					.HAlign(HAlign_Right)
 					.VAlign(VAlign_Center)
 					.Padding(5.0f, 2.0f)
