- Added a dialog when File | Package detects an SDK installation error, but allows user to continue anyway, in case Turnkey is misdetecting SDK

diff --git a/Engine/Source/Editor/MainFrame/Private/Frame/MainFrameActions.cpp b/Engine/Source/Editor/MainFrame/Private/Frame/MainFrameActions.cpp
--- a/Engine/Source/Editor/MainFrame/Private/Frame/MainFrameActions.cpp
+++ b/Engine/Source/Editor/MainFrame/Private/Frame/MainFrameActions.cpp
@@ -44,6 +44,7 @@
 #include "Widgets/Notifications/SNotificationList.h"
 #include "Framework/Commands/GenericCommands.h"
 #include "Dialogs/SOutputLogDialog.h"
+#include "Dialogs/Dialogs.h"
 #include "IUATHelperModule.h"
 #include "Menus/LayoutsMenu.h"
 #include "TargetReceipt.h"
@@ -381,6 +382,31 @@ FString GetCookingOptionalParams()
 }
 
 
+static bool ShowBadSDKDialog(FName IniPlatformName)
+{
+	// Don't show the warning during automation testing; the dlg is modal and blocks
+	if (!GIsAutomationTesting)
+	{
+		FFormatNamedArguments Args;
+		//		Args.Add(TEXT("DisplayName"), PlatformInfo->DisplayName);
+		Args.Add(TEXT("DisplayName"), FText::FromName(IniPlatformName));
+		FText WarningText = FText::Format(LOCTEXT("BadSDK_Message", "The SDK for {DisplayName} appears to not be installed properly, which is needed to generate data. Check the SDK section of the Launch On menu in the main toolbar to update SDK.\n\nWould you like to attempt to continue anyway?"), Args);
+
+		FSuppressableWarningDialog::FSetupInfo Info(
+			WarningText,
+			LOCTEXT("BadSDK_Title", "SDK Not Setup"),
+			TEXT("BadSDKDialog")
+		);
+		Info.ConfirmText = LOCTEXT("BadSDK_Confirm", "Continue");
+		Info.CancelText = LOCTEXT("BadSDK_Cancel", "Cancel");
+		FSuppressableWarningDialog BadSDKgDialog(Info);
+
+		return BadSDKgDialog.ShowModal() != FSuppressableWarningDialog::EResult::Cancel;
+	}
+
+	return true;
+}
+
 void FMainFrameActionCallbacks::CookContent(const FName InPlatformInfoName)
 {
 	const PlatformInfo::FTargetPlatformInfo* const PlatformInfo = PlatformInfo::FindPlatformInfo(InPlatformInfoName);
@@ -403,10 +429,10 @@ void FMainFrameActionCallbacks::CookContent(const FName InPlatformInfoName)
 	}
 
 	// @todo turnkey: run Turnkey here instead of showing SDK!
-	if (PlatformInfo->DataDrivenPlatformInfo->GetSdkStatus() != DDPIPlatformSdkStatus::Valid)
+	if (PlatformInfo->DataDrivenPlatformInfo->GetSdkStatus() != DDPIPlatformSdkStatus::Valid && !ShowBadSDKDialog(PlatformInfo->IniPlatformName))
 	{
-		IMainFrameModule& MainFrameModule = FModuleManager::GetModuleChecked<IMainFrameModule>(TEXT("MainFrame"));
-		MainFrameModule.BroadcastMainFrameSDKNotInstalled(PlatformInfo->Name.ToString(), PlatformInfo->DataDrivenPlatformInfo->SDKTutorial);
+// 		IMainFrameModule& MainFrameModule = FModuleManager::GetModuleChecked<IMainFrameModule>(TEXT("MainFrame"));
+// 		MainFrameModule.BroadcastMainFrameSDKNotInstalled(PlatformInfo->Name.ToString(), PlatformInfo->DataDrivenPlatformInfo->SDKTutorial);
 		return;
 	}
 
@@ -506,12 +532,15 @@ void FMainFrameActionCallbacks::PackageProject( const FName InPlatformInfoName )
 	// @todo turnkey : run turnkey
 	if (PlatformInfo->DataDrivenPlatformInfo->GetSdkStatus() != DDPIPlatformSdkStatus::Valid || (bProjectHasCode && PlatformInfo->DataDrivenPlatformInfo->bUsesHostCompiler && !FSourceCodeNavigation::IsCompilerAvailable()))
 	{
-		IMainFrameModule& MainFrameModule = FModuleManager::GetModuleChecked<IMainFrameModule>(TEXT("MainFrame"));
-		MainFrameModule.BroadcastMainFrameSDKNotInstalled(PlatformInfo->Name.ToString(), PlatformInfo->DataDrivenPlatformInfo->SDKTutorial);
-		TArray<FAnalyticsEventAttribute> ParamArray;
-		ParamArray.Add(FAnalyticsEventAttribute(TEXT("Time"), 0.0));
-		FEditorAnalytics::ReportEvent(TEXT("Editor.Package.Failed"), PlatformInfo->Name.ToString(), bProjectHasCode, EAnalyticsErrorCodes::SDKNotFound, ParamArray);
-		return;
+		if (!ShowBadSDKDialog(PlatformInfo->IniPlatformName))
+		{
+// 			IMainFrameModule& MainFrameModule = FModuleManager::GetModuleChecked<IMainFrameModule>(TEXT("MainFrame"));
+// 			MainFrameModule.BroadcastMainFrameSDKNotInstalled(PlatformInfo->Name.ToString(), PlatformInfo->DataDrivenPlatformInfo->SDKTutorial);
+			TArray<FAnalyticsEventAttribute> ParamArray;
+			ParamArray.Add(FAnalyticsEventAttribute(TEXT("Time"), 0.0));
+			FEditorAnalytics::ReportEvent(TEXT("Editor.Package.Failed"), PlatformInfo->Name.ToString(), bProjectHasCode, EAnalyticsErrorCodes::SDKNotFound, ParamArray);
+			return;
+		}
 	}
 
 	UProjectPackagingSettings* PackagingSettings = Cast<UProjectPackagingSettings>(UProjectPackagingSettings::StaticClass()->GetDefaultObject());
