Prevent Save and Checkout dialog to show up when changing the layout of the Editor.

diff --git a/Engine/Source/Editor/MainFrame/Private/Frame/MainFrameHandler.h b/Engine/Source/Editor/MainFrame/Private/Frame/MainFrameHandler.h
--- a/Engine/Source/Editor/MainFrame/Private/Frame/MainFrameHandler.h
+++ b/Engine/Source/Editor/MainFrame/Private/Frame/MainFrameHandler.h
@@ -132,8 +132,10 @@ public:
 				bOkToExit = bOkToExit && FGlobalTabmanager::Get()->CanCloseManager(TabsToIgnore);
 			}
 
-			// Prompt for save and quit only if we did not launch a gameless rocket exe or are in demo mode
-			if ( FApp::HasProjectName() && !GIsDemoMode )
+
+			IMainFrameModule& MainFrameModule = FModuleManager::LoadModuleChecked<IMainFrameModule>(TEXT("MainFrame"));
+			// Prompt for save and quit only if we did not launch a gameless rocket exe or are in demo mode or we are asking for a close to recreate the Default Main Frame
+			if ( FApp::HasProjectName() && !GIsDemoMode && !MainFrameModule.IsRecreatingDefaultMainFrame())
 			{
 				// Prompt the user to save packages/maps.
 				bool bHadPackagesToSave = false;
diff --git a/Engine/Source/Editor/MainFrame/Private/MainFrameModule.cpp b/Engine/Source/Editor/MainFrame/Private/MainFrameModule.cpp
--- a/Engine/Source/Editor/MainFrame/Private/MainFrameModule.cpp
+++ b/Engine/Source/Editor/MainFrame/Private/MainFrameModule.cpp
@@ -286,6 +286,9 @@ void FMainFrameModule::CreateDefaultMainFrame( const bool bStartImmersive, const
 
 void FMainFrameModule::RecreateDefaultMainFrame(const bool bStartImmersive, const bool bStartPIE)
 {
+	check(!bRecreatingDefaultMainFrame);
+	TGuardValue<bool> GuardRecreatingDefaultMainFrame(bRecreatingDefaultMainFrame, true);
+
 	// Clean previous default main frame
 	if (IsWindowInitialized())
 	{
@@ -298,6 +301,11 @@ void FMainFrameModule::RecreateDefaultMainFrame(const bool bStartImmersive, cons
 	CreateDefaultMainFrame(bStartImmersive, bStartPIE);
 }
 
+bool FMainFrameModule::IsRecreatingDefaultMainFrame() const
+{
+	return bRecreatingDefaultMainFrame;
+}
+
 
 TSharedRef<SWidget> FMainFrameModule::MakeMainMenu(const TSharedPtr<FTabManager>& TabManager, const FName MenuName, FToolMenuContext& ToolMenuContext) const
 {
@@ -494,6 +502,7 @@ void FMainFrameModule::SetApplicationTitleOverride(const FText& NewOverriddenApp
 void FMainFrameModule::StartupModule( )
 {
 	bDelayedShowMainFrame = false;
+	bRecreatingDefaultMainFrame = false;
 
 	MRUFavoritesList = NULL;
 
diff --git a/Engine/Source/Editor/MainFrame/Private/MainFrameModule.h b/Engine/Source/Editor/MainFrame/Private/MainFrameModule.h
--- a/Engine/Source/Editor/MainFrame/Private/MainFrameModule.h
+++ b/Engine/Source/Editor/MainFrame/Private/MainFrameModule.h
@@ -24,6 +24,7 @@ public:
 
 	virtual void CreateDefaultMainFrame( const bool bStartImmersive, const bool bStartPIE ) override;
 	virtual void RecreateDefaultMainFrame(const bool bStartImmersive, const bool bStartPIE) override;
+	virtual bool IsRecreatingDefaultMainFrame() const override;
 	virtual TSharedRef<SWidget> MakeMainMenu(const TSharedPtr<FTabManager>& TabManager, const FName MenuName, FToolMenuContext& ToolMenuContext) const override;
 
 PRAGMA_DISABLE_DEPRECATION_WARNINGS
@@ -227,4 +228,7 @@ private:
 
 	// Allow delaying when to show main frame's window
 	bool bDelayedShowMainFrame;
+
+	// Is recreating Default Main Frame
+	bool bRecreatingDefaultMainFrame;
 };
diff --git a/Engine/Source/Editor/MainFrame/Public/Interfaces/IMainFrameModule.h b/Engine/Source/Editor/MainFrame/Public/Interfaces/IMainFrameModule.h
--- a/Engine/Source/Editor/MainFrame/Public/Interfaces/IMainFrameModule.h
+++ b/Engine/Source/Editor/MainFrame/Public/Interfaces/IMainFrameModule.h
@@ -52,6 +52,11 @@ public:
 	 */
 	virtual void RecreateDefaultMainFrame(const bool bStartImmersive, const bool bStartPIE) = 0;
 
+	/**
+	 * Returns true if the Default Main Frame is being recreated. 
+	 */
+	virtual bool IsRecreatingDefaultMainFrame() const = 0;
+
 	/**
 	 * Generates a menu that includes application global commands, such as "Save All", "Exit", etc.  If you're building
 	 * a menu for your tab, you should call this function to create your menu, passing in an extender object to add your
