Editor layout system: If there is a version bump in the layout system, it used to show a message dialog always. Behavior modified.
If the load request comes from the "Load" UI button, it keep displaying a message dialog to warn the user (and a UELOG(Warning)). If it's from the editor startup, it will only show the UE_LOG(Warning).

Also modified the default layout with the new version.

diff --git a/Engine/Source/Editor/LevelEditor/Private/SLevelEditor.cpp b/Engine/Source/Editor/LevelEditor/Private/SLevelEditor.cpp
--- a/Engine/Source/Editor/LevelEditor/Private/SLevelEditor.cpp
+++ b/Engine/Source/Editor/LevelEditor/Private/SLevelEditor.cpp
@@ -67,6 +67,7 @@
 #include "Elements/Component/ComponentElementLevelEditorSelectionCustomization.h"
 #include "Elements/Component/ComponentElementLevelEditorCommonActionsCustomization.h"
 #include "ThumbnailRendering/ThumbnailManager.h"
+#include <atomic>
 
 #define LOCTEXT_NAMESPACE "SLevelEditor"
 
@@ -1487,17 +1488,25 @@ TSharedRef<SWidget> SLevelEditor::RestoreContentArea( const TSharedRef<SDockTab>
 	const TSharedRef<FTabManager::FLayout> Layout = FLayoutSaveRestore::LoadFromConfig(GEditorLayoutIni,
 		DefaultLayout, OutputCanBeNullptr, RemovedOlderLayoutVersions);
 
+	// On startup, it should not show the dialog to avoid crashes. However, when clicking the "Load" button, it should also show the message dialog
+	// Rather than adding a bool to this function (which would have to be sent from over 10 functions above...), we create a static atomic bool (which is a bit less elegant, but way simpler!), which
+	// will have the same exact effect than sending the bool (and it is also thread safe)
+	static std::atomic<bool> bIsBeingRecreated(false);
 	// If older fields of the layout name (i.e., lower versions than "LevelEditor_Layout_v1.2") were found
 	if (RemovedOlderLayoutVersions.Num() > 0)
 	{
 		// FMessageDialog - Notify the user that the layout version was updated and the current layout uses a deprecated one
-		//const FText TextTitle = LOCTEXT("LevelEditorVersionErrorTitle", "Unreal Level Editor Layout Version Mismatch");
-		const FText TextBody = FText::Format(LOCTEXT("LevelEditorVersionErrorBody", "The expected Unreal Level Editor layout version is \"{0}\", while only version \"{1}\" was found. I.e., the current layout was created with a previous version of Unreal that is deprecated and no longer compatible.\n\nUnreal will continue with the default layout for its current version, the deprecated one has been removed.\n\nYou can create and save your custom layouts with \"Window\"->\"Save Layout\"->\"Save Layout As...\"."),
+		const FText WarningText = FText::Format(LOCTEXT("LevelEditorVersionErrorBody", "The expected Unreal Level Editor layout version is \"{0}\", while only version \"{1}\" was found. I.e., the current layout was created with a previous version of Unreal that is deprecated and no longer compatible.\n\nUnreal will continue with the default layout for its current version, the deprecated one has been removed.\n\nYou can create and save your custom layouts with \"Window\"->\"Save Layout\"->\"Save Layout As...\"."),
 			FText::FromString(LayoutName.ToString()), FText::FromString(RemovedOlderLayoutVersions[0]));
-		//FMessageDialog::Open(EAppMsgType::Ok, TextBody, &TextTitle);
-
-		UE_LOG(LogSlate, Log, TEXT("%s"), *TextBody.ToString());
+		UE_LOG(LogSlate, Warning, TEXT("%s"), *WarningText.ToString());
+		// If user is trying to load a specific layout with "Load", also warn him with a message dialog
+		if (bIsBeingRecreated)
+		{
+			const FText TextTitle = LOCTEXT("LevelEditorVersionErrorTitle", "Unreal Level Editor Layout Version Mismatch");
+			FMessageDialog::Open(EAppMsgType::Ok, WarningText, &TextTitle);
+		}
 	}
+	bIsBeingRecreated = true; // For future loads
 
 	FLayoutExtender LayoutExtender;
 
diff --git a/Engine/Source/Editor/MainFrame/Private/MainFrameModule.cpp b/Engine/Source/Editor/MainFrame/Private/MainFrameModule.cpp
--- a/Engine/Source/Editor/MainFrame/Private/MainFrameModule.cpp
+++ b/Engine/Source/Editor/MainFrame/Private/MainFrameModule.cpp
@@ -71,7 +71,29 @@ const FText StaticGetApplicationTitle( const bool bIncludeGameName )
 /* IMainFrameModule implementation
  *****************************************************************************/
 
-void FMainFrameModule::CreateDefaultMainFrame( const bool bStartImmersive, const bool bStartPIE )
+void FMainFrameModule::CreateDefaultMainFrame(const bool bStartImmersive, const bool bStartPIE)
+{
+	CreateDefaultMainFrameAuxiliary(bStartImmersive, bStartPIE, /*bIsBeingRecreated*/false);
+}
+
+void FMainFrameModule::RecreateDefaultMainFrame(const bool bStartImmersive, const bool bStartPIE)
+{
+	check(!bRecreatingDefaultMainFrame);
+	TGuardValue<bool> GuardRecreatingDefaultMainFrame(bRecreatingDefaultMainFrame, true);
+
+	// Clean previous default main frame
+	if (IsWindowInitialized())
+	{
+		// Clean FSlateApplication
+		FSlateApplication::Get().CloseAllWindowsImmediately();
+		// Clean FGlobalTabmanager
+		FGlobalTabmanager::Get()->CloseAllAreas();
+	}
+	// (Re-)create default main frame
+	CreateDefaultMainFrameAuxiliary(bStartImmersive, bStartPIE, /*bIsBeingRecreated*/true);
+}
+
+void FMainFrameModule::CreateDefaultMainFrameAuxiliary(const bool bStartImmersive, const bool bStartPIE, const bool bIsBeingRecreated)
 {
 	TRACE_CPUPROFILER_EVENT_SCOPE(FMainFrameModule::CreateDefaultMainFrame);
 
@@ -236,10 +258,15 @@ void FMainFrameModule::CreateDefaultMainFrame( const bool bStartImmersive, const
 			if (RemovedOlderLayoutVersions.Num() > 0)
 			{
 				// FMessageDialog - Notify the user that the layout version was updated and the current layout uses a deprecated one
-				const FText TextTitle = LOCTEXT("MainFrameModuleVersionErrorTitle", "Unreal Editor Layout Version Mismatch");
-				const FText TextBody = FText::Format(LOCTEXT("MainFrameModuleVersionErrorBody", "The expected Unreal Editor layout version is \"{0}\", while only version \"{1}\" was found. I.e., the current layout was created with a previous version of Unreal that is deprecated and no longer compatible.\n\nUnreal will continue with the default layout for its current version, the deprecated one has been removed.\n\nYou can create and save your custom layouts with \"Window\"->\"Save Layout\"->\"Save Layout As...\"."),
+				const FText WarningText = FText::Format(LOCTEXT("MainFrameModuleVersionErrorBody", "The expected Unreal Editor layout version is \"{0}\", while only version \"{1}\" was found. I.e., the current layout was created with a previous version of Unreal that is deprecated and no longer compatible.\n\nUnreal will continue with the default layout for its current version, the deprecated one has been removed.\n\nYou can create and save your custom layouts with \"Window\"->\"Save Layout\"->\"Save Layout As...\"."),
 					FText::FromString(LayoutName.ToString()), FText::FromString(RemovedOlderLayoutVersions[0]));
-				FMessageDialog::Open(EAppMsgType::Ok, TextBody, &TextTitle);
+				UE_LOG(LogMainFrame, Warning, TEXT("%s"), *WarningText.ToString());
+				// If user is trying to load a specific layout with "Load", also warn him with a message dialog
+				if (bIsBeingRecreated)
+				{
+					const FText TextTitle = LOCTEXT("MainFrameModuleVersionErrorTitle", "Unreal Editor Layout Version Mismatch");
+					FMessageDialog::Open(EAppMsgType::Ok, WarningText, &TextTitle);
+				}
 			}
 
 			MainFrameContent = FGlobalTabmanager::Get()->RestoreFrom(LoadedLayout, RootWindow, bEmbedTitleAreaContent, OutputCanBeNullptr);
@@ -303,23 +330,6 @@ void FMainFrameModule::CreateDefaultMainFrame( const bool bStartImmersive, const
 	}
 }
 
-void FMainFrameModule::RecreateDefaultMainFrame(const bool bStartImmersive, const bool bStartPIE)
-{
-	check(!bRecreatingDefaultMainFrame);
-	TGuardValue<bool> GuardRecreatingDefaultMainFrame(bRecreatingDefaultMainFrame, true);
-
-	// Clean previous default main frame
-	if (IsWindowInitialized())
-	{
-		// Clean FSlateApplication
-		FSlateApplication::Get().CloseAllWindowsImmediately();
-		// Clean FGlobalTabmanager
-		FGlobalTabmanager::Get()->CloseAllAreas();
-	}
-	// (Re-)create default main frame
-	CreateDefaultMainFrame(bStartImmersive, bStartPIE);
-}
-
 bool FMainFrameModule::IsRecreatingDefaultMainFrame() const
 {
 	return bRecreatingDefaultMainFrame;
diff --git a/Engine/Source/Editor/MainFrame/Private/MainFrameModule.h b/Engine/Source/Editor/MainFrame/Private/MainFrameModule.h
--- a/Engine/Source/Editor/MainFrame/Private/MainFrameModule.h
+++ b/Engine/Source/Editor/MainFrame/Private/MainFrameModule.h
@@ -23,8 +23,16 @@ public:
 
 	// IMainFrameModule interface
 
-	virtual void CreateDefaultMainFrame( const bool bStartImmersive, const bool bStartPIE ) override;
+	virtual void CreateDefaultMainFrame(const bool bStartImmersive, const bool bStartPIE) override;
 	virtual void RecreateDefaultMainFrame(const bool bStartImmersive, const bool bStartPIE) override;
+private:
+	/**
+	 * Shared code between CreateDefaultMainFrame and RecreateDefaultMainFrame
+	 * @param bIsBeingRecreated False if it is being called by first time (CreateDefaultMainFrame), and true if it is being recreated (RecreateDefaultMainFrame). If recreated, it will also display 
+	 */
+	virtual void CreateDefaultMainFrameAuxiliary(const bool bStartImmersive, const bool bStartPIE, const bool bIsBeingRecreated);
+
+public:
 	virtual bool IsRecreatingDefaultMainFrame() const override;
 	virtual TSharedRef<SWidget> MakeMainMenu(const TSharedPtr<FTabManager>& TabManager, const FName MenuName, FToolMenuContext& ToolMenuContext) const override;
 	
