Remove UE4 references.

diff --git a/Engine/Source/Editor/LevelEditor/Private/LevelEditorActions.cpp b/Engine/Source/Editor/LevelEditor/Private/LevelEditorActions.cpp
--- a/Engine/Source/Editor/LevelEditor/Private/LevelEditorActions.cpp
+++ b/Engine/Source/Editor/LevelEditor/Private/LevelEditorActions.cpp
@@ -3154,7 +3154,7 @@ FLevelEditorCommands::FLevelEditorCommands()
 PRAGMA_DISABLE_OPTIMIZATION
 void FLevelEditorCommands::RegisterCommands()
 {
-	UI_COMMAND( BrowseDocumentation, "Documentation...", "Opens the main documentation page, and allows you to search across all UE4 support sites.", EUserInterfaceActionType::Button, FInputChord( EKeys::F1 ) );
+	UI_COMMAND( BrowseDocumentation, "Documentation...", "Opens the main documentation page, and allows you to search across all UE support sites.", EUserInterfaceActionType::Button, FInputChord( EKeys::F1 ) );
 	UI_COMMAND( BrowseAPIReference, "API Reference...", "Opens the API reference documentation", EUserInterfaceActionType::Button, FInputChord() );
 	UI_COMMAND( BrowseCVars, "Console Variables", "Creates an HTML file to browse the console variables and commands (console command 'help')", EUserInterfaceActionType::Button, FInputChord() );
 	UI_COMMAND( BrowseViewportControls, "Viewport Controls...", "Opens the viewport controls cheat sheet", EUserInterfaceActionType::Button, FInputChord() );
diff --git a/Engine/Source/Editor/MainFrame/Private/Frame/MainFrameActions.cpp b/Engine/Source/Editor/MainFrame/Private/Frame/MainFrameActions.cpp
--- a/Engine/Source/Editor/MainFrame/Private/Frame/MainFrameActions.cpp
+++ b/Engine/Source/Editor/MainFrame/Private/Frame/MainFrameActions.cpp
@@ -184,7 +184,7 @@ void FMainFrameCommands::RegisterCommands()
 	UI_COMMAND( VisitAskAQuestionPage, "Ask a Question...", "Have a question?  Go here to ask about anything and everything related to Unreal.", EUserInterfaceActionType::Button, FInputChord() );
 	ActionList->MapAction( VisitAskAQuestionPage, FExecuteAction::CreateStatic( &FMainFrameActionCallbacks::VisitAskAQuestionPage ) );
 
-	UI_COMMAND( VisitSearchForAnswersPage, "Answer Hub...", "Go to the AnswerHub to ask questions, search existing answers, and share your knowledge with other UE4 developers.", EUserInterfaceActionType::Button, FInputChord() );
+	UI_COMMAND( VisitSearchForAnswersPage, "Answer Hub...", "Go to the AnswerHub to ask questions, search existing answers, and share your knowledge with other UE developers.", EUserInterfaceActionType::Button, FInputChord() );
 	ActionList->MapAction( VisitSearchForAnswersPage, FExecuteAction::CreateStatic( &FMainFrameActionCallbacks::VisitSearchForAnswersPage ) );
 
 	UI_COMMAND( VisitSupportWebSite, "Support...", "Navigates to the Unreal Engine Support web site's main page.", EUserInterfaceActionType::Button, FInputChord() );
diff --git a/Engine/Source/Editor/StatsViewer/Classes/TextureStats.h b/Engine/Source/Editor/StatsViewer/Classes/TextureStats.h
--- a/Engine/Source/Editor/StatsViewer/Classes/TextureStats.h
+++ b/Engine/Source/Editor/StatsViewer/Classes/TextureStats.h
@@ -22,7 +22,7 @@ enum ETextureObjectSets
 	TextureObjectSet_SelectedActors			UMETA( DisplayName = "Selected Actor(s)" , ToolTip = "Display texture statistics of selected Actors" ),
 	TextureObjectSet_SelectedMaterials		UMETA( DisplayName = "Selected Materials(s)" , ToolTip = "Display texture statistics of selected Materials" ),
 
-	// @todo: These two stat sets are deprecated as UE4 doesnt support them currently
+	// @todo: These two stat sets are deprecated as UE doesnt support them currently
 	// To recreate them, you will need to re-implement the functionality left behind in
 	// Engine\Source\Editor\Stats\Private\TextureInfo.cpp
 //	TextureObjectSet_CookerStatistics		UMETA( DisplayName = "Cooker Statistics" , ToolTip = "Display texture statistics from a cooker statistics file (.upk)" ),
diff --git a/Engine/Source/Editor/StatsViewer/Private/StatsPages/TextureStatsPage.cpp b/Engine/Source/Editor/StatsViewer/Private/StatsPages/TextureStatsPage.cpp
--- a/Engine/Source/Editor/StatsViewer/Private/StatsPages/TextureStatsPage.cpp
+++ b/Engine/Source/Editor/StatsViewer/Private/StatsPages/TextureStatsPage.cpp
@@ -99,7 +99,7 @@ struct TextureStatsGenerator : public FFindReferencedAssets
 			TexturesToIgnore.Find( MakeWeakObjectPtr( Texture ) ) == INDEX_NONE && // texture is not one that should be ignored
 			( Texture->IsA( UTexture2D::StaticClass() ) || Texture->IsA( UTextureCube::StaticClass() ) ); // texture is valid texture class for stat purposes
 
-#if 0 // @todo TextureInfoInUE4 UTextureCube::GetFace doesn't exist
+#if 0 // @todo TextureInfoInUE UTextureCube::GetFace doesn't exist
 		UTextureCube* CubeTex = Cast<UTextureCube>( Texture );
 		if( CubeTex )
 		{
@@ -249,7 +249,7 @@ struct TextureStatsGenerator : public FFindReferencedAssets
 					Entry->Format = TextureCube->GetPixelFormat();
 					Entry->Type = TEXT("Cube"); 
 
-#if 0 // @todo TextureInfoInUE4 UTextureCube::GetFace doesn't exist.
+#if 0 // @todo TextureInfoInUE UTextureCube::GetFace doesn't exist.
 					// Calculate in game current dimensions 
 					// Use one face of the texture cube to calculate in game size
 					UTexture2D* Face = TextureCube->GetFace(0);
diff --git a/Engine/Source/Runtime/Engine/Private/EngineAnalytics.cpp b/Engine/Source/Runtime/Engine/Private/EngineAnalytics.cpp
--- a/Engine/Source/Runtime/Engine/Private/EngineAnalytics.cpp
+++ b/Engine/Source/Runtime/Engine/Private/EngineAnalytics.cpp
@@ -45,10 +45,10 @@ static TSharedPtr<IAnalyticsProviderET> CreateEpicAnalyticsProvider()
 			!FEngineBuildSettings::IsInternalBuild();	// Internal Epic build
 		const TCHAR* BuildTypeStr = bUseReleaseAccount ? TEXT("Release") : TEXT("Dev");
 
-		FString UE4TypeOverride;
-		bool bHasOverride = GConfig->GetString(TEXT("Analytics"), TEXT("UE4TypeOverride"), UE4TypeOverride, GEngineIni);
-		const TCHAR* UE4TypeStr = bHasOverride ? *UE4TypeOverride : FEngineBuildSettings::IsPerforceBuild() ? TEXT("Perforce") : TEXT("UnrealEngine");
-		Config.APIKeyET = FString::Printf(TEXT("UEEditor.%s.%s"), UE4TypeStr, BuildTypeStr);
+		FString UETypeOverride;
+		bool bHasOverride = GConfig->GetString(TEXT("Analytics"), TEXT("UE4TypeOverride"), UETypeOverride, GEngineIni);
+		const TCHAR* UETypeStr = bHasOverride ? *UETypeOverride : FEngineBuildSettings::IsPerforceBuild() ? TEXT("Perforce") : TEXT("UnrealEngine");
+		Config.APIKeyET = FString::Printf(TEXT("UEEditor.%s.%s"), UETypeStr, BuildTypeStr);
 	}
 	Config.APIServerET = TEXT("https://datarouter.ol.epicgames.com/");
 	Config.AppEnvironment = TEXT("datacollector-binary");
diff --git a/Engine/Source/Runtime/Engine/Public/EngineAnalytics.h b/Engine/Source/Runtime/Engine/Public/EngineAnalytics.h
--- a/Engine/Source/Runtime/Engine/Public/EngineAnalytics.h
+++ b/Engine/Source/Runtime/Engine/Public/EngineAnalytics.h
@@ -13,7 +13,7 @@ struct FAnalyticsEventAttribute;
 /**
  * The public interface for the editor's analytics provider singleton.
  * 
- * WARNING: This is an analytics provider instance that is created whenever UE4 editor is launched. 
+ * WARNING: This is an analytics provider instance that is created whenever UE editor is launched. 
  * It is intended ONLY for use by Epic Games. This is NOT intended for games to send 
  * game-specific telemetry. Create your own provider instance for your game and configure
  * it independently.
