Level Editor Menu Re-org, Part 1

diff --git a/Engine/Plugins/Animation/LiveLink/Source/LiveLinkEditor/Private/LiveLinkEditorModule.cpp b/Engine/Plugins/Animation/LiveLink/Source/LiveLinkEditor/Private/LiveLinkEditorModule.cpp
--- a/Engine/Plugins/Animation/LiveLink/Source/LiveLinkEditor/Private/LiveLinkEditorModule.cpp
+++ b/Engine/Plugins/Animation/LiveLink/Source/LiveLinkEditor/Private/LiveLinkEditorModule.cpp
@@ -191,11 +191,9 @@ private:
 		FTabSpawnerEntry& SpawnerEntry = FGlobalTabmanager::Get()->RegisterNomadTabSpawner(LiveLinkClientTabName, FOnSpawnTab::CreateStatic(&FLiveLinkEditorModule::SpawnLiveLinkTab, StyleSet))
 			.SetDisplayName(LOCTEXT("LiveLinkTabTitle", "Live Link"))
 			.SetTooltipText(LOCTEXT("LiveLinkTabTooltipText", "Open the Live Link streaming manager tab."))
+			.SetGroup(WorkspaceMenu::GetMenuStructure().GetLevelEditorVirtualProductionCategory())
 			.SetIcon(FSlateIcon(StyleSet->GetStyleSetName(), "LiveLinkClient.Common.Icon.Small"));
 
-		const IWorkspaceMenuStructure& MenuStructure = WorkspaceMenu::GetMenuStructure();
-		SpawnerEntry.SetGroup(MenuStructure.GetLevelEditorCategory());
-
 		bHasRegisteredTabSpawners = true;
 	}
 
diff --git a/Engine/Plugins/Animation/LiveLinkCurveDebugUI/Source/LiveLinkCurveDebugUI/Private/LiveLinkCurveDebugUIModule.cpp b/Engine/Plugins/Animation/LiveLinkCurveDebugUI/Source/LiveLinkCurveDebugUI/Private/LiveLinkCurveDebugUIModule.cpp
--- a/Engine/Plugins/Animation/LiveLinkCurveDebugUI/Source/LiveLinkCurveDebugUI/Private/LiveLinkCurveDebugUIModule.cpp
+++ b/Engine/Plugins/Animation/LiveLinkCurveDebugUI/Source/LiveLinkCurveDebugUI/Private/LiveLinkCurveDebugUIModule.cpp
@@ -149,7 +149,7 @@ void FLiveLinkCurveDebugUIModule::RegisterTabSpawner()
 	//Register with the Developer Tools Menu
 #if WITH_EDITOR
 	const IWorkspaceMenuStructure& MenuStructure = WorkspaceMenu::GetMenuStructure();
-	SpawnerEntry.SetGroup(MenuStructure.GetDeveloperToolsMiscCategory());
+	SpawnerEntry.SetGroup(MenuStructure.GetDeveloperToolsDebugCategory());
 #endif //WITH_EDITOR
 }
 
diff --git a/Engine/Plugins/Cameras/CameraShakePreviewer/Source/CameraShakePreviewer/Private/CameraShakePreviewerModule.cpp b/Engine/Plugins/Cameras/CameraShakePreviewer/Source/CameraShakePreviewer/Private/CameraShakePreviewerModule.cpp
--- a/Engine/Plugins/Cameras/CameraShakePreviewer/Source/CameraShakePreviewer/Private/CameraShakePreviewerModule.cpp
+++ b/Engine/Plugins/Cameras/CameraShakePreviewer/Source/CameraShakePreviewer/Private/CameraShakePreviewerModule.cpp
@@ -92,7 +92,7 @@ void FCameraShakePreviewerModule::RegisterEditorTab()
 			.SetDisplayName(LOCTEXT("CameraShakePreviewer", "Camera Shake Previewer"))
 			.SetTooltipText(LOCTEXT("CameraShakePreviewerTooltipText", "Open the camera shake preview panel."))
 			.SetIcon(Icon)
-			.SetGroup(MenuStructure.GetLevelEditorCategory());
+			.SetGroup(WorkspaceMenu::GetMenuStructure().GetLevelEditorCinematicsCategory());
 	});
 }
 
diff --git a/Engine/Plugins/Compositing/Composure/Source/ComposureLayersEditor/Private/CompElementEditorModule.cpp b/Engine/Plugins/Compositing/Composure/Source/ComposureLayersEditor/Private/CompElementEditorModule.cpp
--- a/Engine/Plugins/Compositing/Composure/Source/ComposureLayersEditor/Private/CompElementEditorModule.cpp
+++ b/Engine/Plugins/Compositing/Composure/Source/ComposureLayersEditor/Private/CompElementEditorModule.cpp
@@ -279,7 +279,7 @@ void FCompElementEditorModule::RegisterEditorTab()
 		LevelEditorTabManager->RegisterTabSpawner(CompElementEditor_Impl::ComposureLayersTabName, FOnSpawnTab::CreateStatic(&FCompElementEditorModule::SpawnComposureLayersTab))
 			.SetDisplayName(NSLOCTEXT("LevelEditorTabs", "LevelEditorComposureLayerBrowser", "Composure Compositing"))
 			.SetTooltipText(NSLOCTEXT("LevelEditorTabs", "LevelEditorComposureLayerBrowserTooltipText", "Open the Composure compositing tab."))
-			.SetGroup(WorkspaceMenu::GetMenuStructure().GetLevelEditorCategory())
+			.SetGroup(WorkspaceMenu::GetMenuStructure().GetLevelEditorVirtualProductionCategory())
 			.SetIcon(LayersIcon);
 	});
 }
diff --git a/Engine/Plugins/Developer/Concert/DisasterRecoveryClient/Source/DisasterRecoveryClient/Private/DisasterRecoveryClientModule.cpp b/Engine/Plugins/Developer/Concert/DisasterRecoveryClient/Source/DisasterRecoveryClient/Private/DisasterRecoveryClientModule.cpp
--- a/Engine/Plugins/Developer/Concert/DisasterRecoveryClient/Source/DisasterRecoveryClient/Private/DisasterRecoveryClientModule.cpp
+++ b/Engine/Plugins/Developer/Concert/DisasterRecoveryClient/Source/DisasterRecoveryClient/Private/DisasterRecoveryClientModule.cpp
@@ -205,7 +205,7 @@ private:
 #if WITH_EDITOR
 	void RegisterTabSpawner()
 	{
-		TSharedPtr<FWorkspaceItem> WorkspaceGroup = WorkspaceMenu::GetMenuStructure().GetDeveloperToolsMiscCategory();
+		TSharedPtr<FWorkspaceItem> WorkspaceGroup = WorkspaceMenu::GetMenuStructure().GetToolsCategory();
 		if (bTabSpawnerRegistrated)
 		{
 			UnregisterTabSpawner();
diff --git a/Engine/Plugins/Developer/Concert/MultiUserClient/Source/MultiUserClient/Private/MultiUserClientModule.cpp b/Engine/Plugins/Developer/Concert/MultiUserClient/Source/MultiUserClient/Private/MultiUserClientModule.cpp
--- a/Engine/Plugins/Developer/Concert/MultiUserClient/Source/MultiUserClient/Private/MultiUserClientModule.cpp
+++ b/Engine/Plugins/Developer/Concert/MultiUserClient/Source/MultiUserClient/Private/MultiUserClientModule.cpp
@@ -782,7 +782,7 @@ private:
 		FModuleManager::LoadModuleChecked<IEditorStyleModule>("EditorStyle");
 
 #if WITH_EDITOR
-		RegisterTabSpawner(WorkspaceMenu::GetMenuStructure().GetDeveloperToolsMiscCategory());
+		RegisterTabSpawner(WorkspaceMenu::GetMenuStructure().GetLevelEditorCategory());
 
 		// Register Workspace view
 		RegisterWorkspaceUI();
diff --git a/Engine/Plugins/Editor/AssetManagerEditor/Source/AssetManagerEditor/Private/AssetManagerEditorModule.cpp b/Engine/Plugins/Editor/AssetManagerEditor/Source/AssetManagerEditor/Private/AssetManagerEditorModule.cpp
--- a/Engine/Plugins/Editor/AssetManagerEditor/Source/AssetManagerEditor/Private/AssetManagerEditorModule.cpp
+++ b/Engine/Plugins/Editor/AssetManagerEditor/Source/AssetManagerEditor/Private/AssetManagerEditorModule.cpp
@@ -498,7 +498,7 @@ void FAssetManagerEditorModule::StartupModule()
 		FGlobalTabmanager::Get()->RegisterNomadTabSpawner(AssetAuditTabName, FOnSpawnTab::CreateRaw(this, &FAssetManagerEditorModule::SpawnAssetAuditTab))
 			.SetDisplayName(LOCTEXT("AssetAuditTitle", "Asset Audit"))
 			.SetTooltipText(LOCTEXT("AssetAuditTooltip", "Open Asset Audit window, allows viewing detailed information about assets."))
-			.SetGroup(WorkspaceMenu::GetMenuStructure().GetDeveloperToolsMiscCategory())
+			.SetGroup(WorkspaceMenu::GetMenuStructure().GetDeveloperToolsAuditCategory())
 			.SetIcon(FSlateIcon(FEditorStyle::Get().GetStyleSetName(), "GenericEditor.Tabs.Properties"));
 
 		FGlobalTabmanager::Get()->RegisterNomadTabSpawner(ReferenceViewerTabName, FOnSpawnTab::CreateRaw(this, &FAssetManagerEditorModule::SpawnReferenceViewerTab))
diff --git a/Engine/Plugins/Editor/DataValidation/Source/DataValidation/Private/DataValidationModule.cpp b/Engine/Plugins/Editor/DataValidation/Source/DataValidation/Private/DataValidationModule.cpp
--- a/Engine/Plugins/Editor/DataValidation/Source/DataValidation/Private/DataValidationModule.cpp
+++ b/Engine/Plugins/Editor/DataValidation/Source/DataValidation/Private/DataValidationModule.cpp
@@ -114,8 +114,8 @@ void FDataValidationModule::ShutdownModule()
 void FDataValidationModule::RegisterMenus()
 {
 	FToolMenuOwnerScoped OwnerScoped(this);
-	UToolMenu* Menu = UToolMenus::Get()->ExtendMenu("LevelEditor.MainMenu.File");
-	FToolMenuSection& Section = Menu->AddSection("DataValidation", LOCTEXT("DataValidation", "DataValidation"), FToolMenuInsert("FileLoadAndSave", EToolMenuInsertType::After));
+	UToolMenu* Menu = UToolMenus::Get()->ExtendMenu("LevelEditor.MainMenu.Tools");
+	FToolMenuSection& Section = Menu->AddSection("DataValidation", LOCTEXT("DataValidation", "DataValidation"));
 	Section.AddEntry(FToolMenuEntry::InitMenuEntry(
 		"ValidateData",
 		TAttribute<FText>::Create(&Menu_ValidateDataGetTitle),
diff --git a/Engine/Plugins/Editor/EditorDebugTools/Source/EditorDebugTools/Private/EditorDebugTools.cpp b/Engine/Plugins/Editor/EditorDebugTools/Source/EditorDebugTools/Private/EditorDebugTools.cpp
--- a/Engine/Plugins/Editor/EditorDebugTools/Source/EditorDebugTools/Private/EditorDebugTools.cpp
+++ b/Engine/Plugins/Editor/EditorDebugTools/Source/EditorDebugTools/Private/EditorDebugTools.cpp
@@ -77,7 +77,7 @@ void FEditorDebugToolsModule::StartupModule()
 			.SetDisplayName(NSLOCTEXT("Toolbox", "Modules", "Modules"))
 			.SetTooltipText(NSLOCTEXT("Toolbox", "ModulesTooltipText", "Open the Modules tab."))
 			.SetIcon(FSlateIcon(FEditorDebugToolsStyle::GetStyleSetName(), "Modules.TabIcon"))
-			.SetGroup(WorkspaceMenu::GetMenuStructure().GetDeveloperToolsMiscCategory());
+			.SetGroup(WorkspaceMenu::GetMenuStructure().GetDeveloperToolsDebugCategory());
 	}
 }
 
diff --git a/Engine/Plugins/Editor/FacialAnimation/Source/FacialAnimationEditor/Private/FacialAnimationEditorModule.cpp b/Engine/Plugins/Editor/FacialAnimation/Source/FacialAnimationEditor/Private/FacialAnimationEditorModule.cpp
--- a/Engine/Plugins/Editor/FacialAnimation/Source/FacialAnimationEditor/Private/FacialAnimationEditorModule.cpp
+++ b/Engine/Plugins/Editor/FacialAnimation/Source/FacialAnimationEditor/Private/FacialAnimationEditorModule.cpp
@@ -73,7 +73,7 @@ void FFacialAnimationEditorModule::StartupModule()
 		FTabSpawnerEntry& TabSpawner = FGlobalTabmanager::Get()->RegisterNomadTabSpawner(FacialAnimationBulkImporterTabName, FOnSpawnTab::CreateStatic(&Local::SpawnFacialAnimationBulkImporterTab))
 			.SetDisplayName(LOCTEXT("FacialAnimationBulkImporterTabTitle", "Facial Anim Importer"))
 			.SetTooltipText(LOCTEXT("FacialAnimationBulkImporterTooltipText", "Open the Facial Animation Bulk Importer tab."))
-			.SetGroup(WorkspaceMenu::GetMenuStructure().GetDeveloperToolsMiscCategory())
+			.SetGroup(WorkspaceMenu::GetMenuStructure().GetToolsCategory())
 			.SetIcon(FSlateIcon(FAppStyle::GetAppStyleSetName(), "Icons..Import"));
 
 		// register for when persona is loaded
diff --git a/Engine/Plugins/Editor/MaterialAnalyzer/Source/MaterialAnalyzerModule.cpp b/Engine/Plugins/Editor/MaterialAnalyzer/Source/MaterialAnalyzerModule.cpp
--- a/Engine/Plugins/Editor/MaterialAnalyzer/Source/MaterialAnalyzerModule.cpp
+++ b/Engine/Plugins/Editor/MaterialAnalyzer/Source/MaterialAnalyzerModule.cpp
@@ -43,7 +43,7 @@ void FMaterialAnalyzerModule::StartupModule()
 	FGlobalTabmanager::Get()->RegisterNomadTabSpawner(
 		MaterialAnalyzerName,
 		FOnSpawnTab::CreateRaw(this, &FMaterialAnalyzerModule::SpawnMaterialAnalyzerTab))
-		.SetGroup(WorkspaceMenu::GetMenuStructure().GetDeveloperToolsMiscCategory())
+		.SetGroup(WorkspaceMenu::GetMenuStructure().GetDeveloperToolsAuditCategory())
 		.SetDisplayName(LOCTEXT("TabTitle", "Material Analyzer"))
 		.SetTooltipText(LOCTEXT("TooltipText", "Opens Material Analyzer tool."))
 		.SetIcon(FSlateIcon(FEditorStyle::GetStyleSetName(), "MaterialEditor.ToggleMaterialStats.Tab"));
diff --git a/Engine/Plugins/Editor/StylusInput/Source/StylusInput/Private/StylusInputModule.cpp b/Engine/Plugins/Editor/StylusInput/Source/StylusInput/Private/StylusInputModule.cpp
--- a/Engine/Plugins/Editor/StylusInput/Source/StylusInput/Private/StylusInputModule.cpp
+++ b/Engine/Plugins/Editor/StylusInput/Source/StylusInput/Private/StylusInputModule.cpp
@@ -55,7 +55,7 @@ void UStylusInputSubsystem::Initialize(FSubsystemCollectionBase& Collection)
 		FOnSpawnTab::CreateUObject(this, &UStylusInputSubsystem::OnSpawnPluginTab))
 		.SetDisplayName(LOCTEXT("DebugTabTitle", "Stylus Input Debug"))
 		.SetTooltipText(LOCTEXT("DebugTabTooltip", "Debug panel to display current values of stylus inputs."))
-		.SetGroup(MenuStructure.GetDeveloperToolsMiscCategory());
+		.SetGroup(MenuStructure.GetDeveloperToolsDebugCategory());
 }
 
 void UStylusInputSubsystem::Deinitialize()
diff --git a/Engine/Plugins/Experimental/CodeEditor/Source/CodeEditor/Private/CodeEditor.cpp b/Engine/Plugins/Experimental/CodeEditor/Source/CodeEditor/Private/CodeEditor.cpp
--- a/Engine/Plugins/Experimental/CodeEditor/Source/CodeEditor/Private/CodeEditor.cpp
+++ b/Engine/Plugins/Experimental/CodeEditor/Source/CodeEditor/Private/CodeEditor.cpp
@@ -39,7 +39,7 @@ public:
 	{
 		if (UToolMenus::IsToolMenuUIEnabled())
 		{
-			UToolMenu* Menu = UToolMenus::Get()->ExtendMenu("LevelEditor.MainMenu.File");
+			UToolMenu* Menu = UToolMenus::Get()->ExtendMenu("LevelEditor.MainMenu.Tools");
 			FToolMenuSection& Section = Menu->FindOrAddSection("FileProject");
 
 			FToolMenuOwnerScoped OwnerScoped(this);
diff --git a/Engine/Plugins/Experimental/ControlRig/Source/ControlRigEditor/Private/ControlRigEditorModule.cpp b/Engine/Plugins/Experimental/ControlRig/Source/ControlRigEditor/Private/ControlRigEditorModule.cpp
--- a/Engine/Plugins/Experimental/ControlRig/Source/ControlRigEditor/Private/ControlRigEditorModule.cpp
+++ b/Engine/Plugins/Experimental/ControlRig/Source/ControlRigEditor/Private/ControlRigEditorModule.cpp
@@ -278,7 +278,7 @@ void FControlRigEditorModule::StartupModule()
 		FGlobalTabmanager::Get()->RegisterNomadTabSpawner("HierarchicalProfiler", FOnSpawnTab::CreateStatic(&SpawnRigProfiler))
 			.SetDisplayName(NSLOCTEXT("UnrealEditor", "HierarchicalProfilerTab", "Hierarchical Profiler"))
 			.SetTooltipText(NSLOCTEXT("UnrealEditor", "HierarchicalProfilerTooltip", "Open the Hierarchical Profiler tab."))
-			.SetGroup(WorkspaceMenu::GetMenuStructure().GetDeveloperToolsMiscCategory())
+			.SetGroup(WorkspaceMenu::GetMenuStructure().GetDeveloperToolsProfilingCategory())
 			.SetIcon(FSlateIcon(TEXT("ControlRigEditorStyle"), TEXT("ControlRig.RigUnit")));
 
 		FGlobalTabmanager::Get()->RegisterNomadTabSpawner(IControlRigEditorModule::ControlRigPoseTab, FOnSpawnTab::CreateStatic(&SpawnPoseTab))
diff --git a/Engine/Plugins/Experimental/PythonScriptPlugin/Source/PythonScriptPlugin/Private/PythonScriptPlugin.cpp b/Engine/Plugins/Experimental/PythonScriptPlugin/Source/PythonScriptPlugin/Private/PythonScriptPlugin.cpp
--- a/Engine/Plugins/Experimental/PythonScriptPlugin/Source/PythonScriptPlugin/Private/PythonScriptPlugin.cpp
+++ b/Engine/Plugins/Experimental/PythonScriptPlugin/Source/PythonScriptPlugin/Private/PythonScriptPlugin.cpp
@@ -352,8 +352,8 @@ private:
 	void RegisterMenus()
 	{
 		FToolMenuOwnerScoped OwnerScoped(this);
-		UToolMenu* Menu = UToolMenus::Get()->ExtendMenu("LevelEditor.MainMenu.File");
-		FToolMenuSection& Section = Menu->AddSection("Python", LOCTEXT("Python", "Python"), FToolMenuInsert("FileLoadAndSave", EToolMenuInsertType::After));
+		UToolMenu* Menu = UToolMenus::Get()->ExtendMenu("LevelEditor.MainMenu.Tools");
+		FToolMenuSection& Section = Menu->AddSection("Python", LOCTEXT("Python", "Python"));//, FToolMenuInsert("FileLoadAndSave", EToolMenuInsertType::After));
 		Section.AddMenuEntry(
 			"OpenPython",
 			LOCTEXT("OpenPython", "Execute Python Script..."),
diff --git a/Engine/Plugins/Experimental/RemoteSession/Source/RemoteSessionEditor/Private/Widgets/SRemoteSessionStream.cpp b/Engine/Plugins/Experimental/RemoteSession/Source/RemoteSessionEditor/Private/Widgets/SRemoteSessionStream.cpp
--- a/Engine/Plugins/Experimental/RemoteSession/Source/RemoteSessionEditor/Private/Widgets/SRemoteSessionStream.cpp
+++ b/Engine/Plugins/Experimental/RemoteSession/Source/RemoteSessionEditor/Private/Widgets/SRemoteSessionStream.cpp
@@ -79,7 +79,7 @@ void SRemoteSessionStream::RegisterNomadTabSpawner()
 		LevelEditorTabManager->RegisterTabSpawner(RemoteSessionStream::RemoteSessionStreamApp, FOnSpawnTab::CreateStatic(&RemoteSessionStream::CreateMediaFrameworkCaptureCameraViewportTab))
 			.SetDisplayName(LOCTEXT("TabTitle", "Remote Session Stream"))
 			.SetTooltipText(LOCTEXT("TooltipText", "Stream a particular UMG to a Remote Session app."))
-			.SetGroup(WorkspaceMenu::GetMenuStructure().GetDeveloperToolsMiscCategory())
+			.SetGroup(WorkspaceMenu::GetMenuStructure().GetToolsCategory())
 			.SetIcon(FSlateIcon(FRemoteSessionEditorStyle::GetStyleSetName(), "TabIcons.RemoteSession.Small"));
 	};
 
diff --git a/Engine/Plugins/Experimental/VirtualCamera/Source/VirtualCameraEditor/Private/VirtualCameraEditorModule.cpp b/Engine/Plugins/Experimental/VirtualCamera/Source/VirtualCameraEditor/Private/VirtualCameraEditorModule.cpp
--- a/Engine/Plugins/Experimental/VirtualCamera/Source/VirtualCameraEditor/Private/VirtualCameraEditorModule.cpp
+++ b/Engine/Plugins/Experimental/VirtualCamera/Source/VirtualCameraEditor/Private/VirtualCameraEditorModule.cpp
@@ -27,7 +27,7 @@ public:
 	{
 		FVirtualCameraEditorStyle::Register();
 		const IWorkspaceMenuStructure& MenuStructure = WorkspaceMenu::GetMenuStructure();
-		TSharedRef<FWorkspaceItem> DeveloperToolsGroup = MenuStructure.GetDeveloperToolsMiscCategory();
+		TSharedRef<FWorkspaceItem> DeveloperToolsGroup = MenuStructure.GetLevelEditorVirtualProductionCategory();
 		SVirtualCameraTab::RegisterNomadTabSpawner(DeveloperToolsGroup);
 
 		RegisterSettings();
diff --git a/Engine/Plugins/Experimental/VirtualProductionUtilities/Source/VPUtilitiesEditor/Private/VPUtilitiesEditorModule.cpp b/Engine/Plugins/Experimental/VirtualProductionUtilities/Source/VPUtilitiesEditor/Private/VPUtilitiesEditorModule.cpp
--- a/Engine/Plugins/Experimental/VirtualProductionUtilities/Source/VPUtilitiesEditor/Private/VPUtilitiesEditorModule.cpp
+++ b/Engine/Plugins/Experimental/VirtualProductionUtilities/Source/VPUtilitiesEditor/Private/VPUtilitiesEditorModule.cpp
@@ -38,15 +38,7 @@ void FVPUtilitiesEditorModule::StartupModule()
 	CustomUIHandler.Reset(NewObject<UVPCustomUIHandler>());
 	CustomUIHandler->Init();
 
-	{
-		const IWorkspaceMenuStructure& MenuStructure = WorkspaceMenu::GetMenuStructure();
-		TSharedRef<FWorkspaceItem> MediaBrowserGroup = MenuStructure.GetDeveloperToolsMiscCategory()->GetParent()->AddGroup(
-			LOCTEXT("WorkspaceMenu_VirtualProductionCategory", "Virtual Production"),
-			FSlateIcon(),
-			true);
-
-		SGenlockProviderTab::RegisterNomadTabSpawner(MediaBrowserGroup);
-	}
+	SGenlockProviderTab::RegisterNomadTabSpawner(WorkspaceMenu::GetMenuStructure().GetLevelEditorVirtualProductionCategory());
 
 	RegisterSettings();
 
diff --git a/Engine/Plugins/Importers/USDImporter/Source/USDStageEditor/Private/USDStageEditorModule.cpp b/Engine/Plugins/Importers/USDImporter/Source/USDStageEditor/Private/USDStageEditorModule.cpp
--- a/Engine/Plugins/Importers/USDImporter/Source/USDStageEditor/Private/USDStageEditorModule.cpp
+++ b/Engine/Plugins/Importers/USDImporter/Source/USDStageEditor/Private/USDStageEditorModule.cpp
@@ -53,7 +53,7 @@ public:
 				LevelEditorTabManager->RegisterTabSpawner( TEXT("USDStage"), FOnSpawnTab::CreateStatic( &FUsdStageEditorModule::SpawnUsdStageTab ) )
 					.SetDisplayName( LOCTEXT( "USDStage", "USD Stage" ) )
 					.SetTooltipText( LOCTEXT( "USDStageTab", "Open USD Stage tab" ) )
-					.SetGroup( WorkspaceMenu::GetMenuStructure().GetLevelEditorCategory() )
+					.SetGroup(WorkspaceMenu::GetMenuStructure().GetLevelEditorVirtualProductionCategory())
 					.SetIcon( LayersIcon );
 			});
 
diff --git a/Engine/Plugins/Media/MediaFrameworkUtilities/Source/MediaFrameworkUtilitiesEditor/Private/MediaFrameworkUtilitiesEditorModule.cpp b/Engine/Plugins/Media/MediaFrameworkUtilities/Source/MediaFrameworkUtilitiesEditor/Private/MediaFrameworkUtilitiesEditorModule.cpp
--- a/Engine/Plugins/Media/MediaFrameworkUtilities/Source/MediaFrameworkUtilitiesEditor/Private/MediaFrameworkUtilitiesEditorModule.cpp
+++ b/Engine/Plugins/Media/MediaFrameworkUtilities/Source/MediaFrameworkUtilitiesEditor/Private/MediaFrameworkUtilitiesEditorModule.cpp
@@ -78,7 +78,7 @@ public:
 
 			{
 				const IWorkspaceMenuStructure& MenuStructure = WorkspaceMenu::GetMenuStructure();
-				TSharedRef<FWorkspaceItem> MediaBrowserGroup = MenuStructure.GetDeveloperToolsMiscCategory()->GetParent()->AddGroup(
+				TSharedRef<FWorkspaceItem> MediaBrowserGroup = MenuStructure.GetLevelEditorVirtualProductionCategory()->AddGroup(
 					LOCTEXT("WorkspaceMenu_MediaCategory", "Media"),
 					FSlateIcon(),
 					true);
diff --git a/Engine/Plugins/MovieScene/MatineeToLevelSequence/Source/MatineeToLevelSequence/Private/MatineeToLevelSequenceModule.cpp b/Engine/Plugins/MovieScene/MatineeToLevelSequence/Source/MatineeToLevelSequence/Private/MatineeToLevelSequenceModule.cpp
--- a/Engine/Plugins/MovieScene/MatineeToLevelSequence/Source/MatineeToLevelSequence/Private/MatineeToLevelSequenceModule.cpp
+++ b/Engine/Plugins/MovieScene/MatineeToLevelSequence/Source/MatineeToLevelSequence/Private/MatineeToLevelSequenceModule.cpp
@@ -110,7 +110,7 @@ protected:
 				.SetDisplayName(LOCTEXT("MatineeCameraShakeConverter", "Matinee Camera Shake Converter"))
 				.SetTooltipText(LOCTEXT("MatineeCameraShakeConverterTooltipText", "Open the legacy camera shake converter tool."))
 				.SetIcon(Icon)
-				.SetGroup(MenuStructure.GetDeveloperToolsMiscCategory());
+				.SetGroup(MenuStructure.GetLevelEditorCinematicsCategory());
 			});
 	}
 
diff --git a/Engine/Plugins/VirtualProduction/StageMonitoring/Source/StageMonitorEditor/Private/StageMonitorEditorModule.cpp b/Engine/Plugins/VirtualProduction/StageMonitoring/Source/StageMonitorEditor/Private/StageMonitorEditorModule.cpp
--- a/Engine/Plugins/VirtualProduction/StageMonitoring/Source/StageMonitorEditor/Private/StageMonitorEditorModule.cpp
+++ b/Engine/Plugins/VirtualProduction/StageMonitoring/Source/StageMonitorEditor/Private/StageMonitorEditorModule.cpp
@@ -20,7 +20,7 @@ DEFINE_LOG_CATEGORY(LogStageMonitorEditor);
 
 void FStageMonitorEditorModule::StartupModule()
 {
-	SStageMonitorPanel::RegisterNomadTabSpawner(WorkspaceMenu::GetMenuStructure().GetDeveloperToolsMiscCategory());
+	SStageMonitorPanel::RegisterNomadTabSpawner(WorkspaceMenu::GetMenuStructure().GetLevelEditorVirtualProductionCategory());
 
 	ISettingsModule* SettingsModule = FModuleManager::GetModulePtr<ISettingsModule>("Settings");
 	if (SettingsModule)
diff --git a/Engine/Plugins/VirtualProduction/TimedDataMonitor/Source/TimedDataMonitorEditor/Private/TimedDataMonitorEditorModule.cpp b/Engine/Plugins/VirtualProduction/TimedDataMonitor/Source/TimedDataMonitorEditor/Private/TimedDataMonitorEditorModule.cpp
--- a/Engine/Plugins/VirtualProduction/TimedDataMonitor/Source/TimedDataMonitorEditor/Private/TimedDataMonitorEditorModule.cpp
+++ b/Engine/Plugins/VirtualProduction/TimedDataMonitor/Source/TimedDataMonitorEditor/Private/TimedDataMonitorEditorModule.cpp
@@ -20,7 +20,7 @@ class FTimedDataMonitorEditorModule : public IModuleInterface
 public:
 	virtual void StartupModule() override
 	{
-		STimedDataMonitorPanel::RegisterNomadTabSpawner(WorkspaceMenu::GetMenuStructure().GetDeveloperToolsMiscCategory());
+		STimedDataMonitorPanel::RegisterNomadTabSpawner(WorkspaceMenu::GetMenuStructure().GetToolsCategory());
 
 		if (ISettingsModule* SettingsModule = FModuleManager::GetModulePtr<ISettingsModule>("Settings"))
 		{
diff --git a/Engine/Source/Developer/DeviceManager/Private/DeviceManagerModule.cpp b/Engine/Source/Developer/DeviceManager/Private/DeviceManagerModule.cpp
--- a/Engine/Source/Developer/DeviceManager/Private/DeviceManagerModule.cpp
+++ b/Engine/Source/Developer/DeviceManager/Private/DeviceManagerModule.cpp
@@ -42,7 +42,7 @@ public:
 			.SetIcon(FSlateIcon(FEditorStyle::GetStyleSetName(), "DeviceDetails.TabIcon"));
 
 #if WITH_EDITOR
-		TabSpawnerEntry.SetGroup(WorkspaceMenu::GetMenuStructure().GetDeveloperToolsMiscCategory());
+		TabSpawnerEntry.SetGroup(WorkspaceMenu::GetMenuStructure().GetDeveloperToolsPlatformsCategory());
 #else
 		TabSpawnerEntry.SetGroup(WorkspaceMenu::GetMenuStructure().GetToolsCategory());
 #endif
diff --git a/Engine/Source/Developer/LogVisualizer/Private/LogVisualizerModule.cpp b/Engine/Source/Developer/LogVisualizer/Private/LogVisualizerModule.cpp
--- a/Engine/Source/Developer/LogVisualizer/Private/LogVisualizerModule.cpp
+++ b/Engine/Source/Developer/LogVisualizer/Private/LogVisualizerModule.cpp
@@ -53,7 +53,7 @@ void FLogVisualizerModule::StartupModule()
 	FGlobalTabmanager::Get()->RegisterNomadTabSpawner(
 		VisualLoggerTabName, 
 		FOnSpawnTab::CreateRaw(this, &FLogVisualizerModule::SpawnLogVisualizerTab))
-		.SetGroup(WorkspaceMenu::GetMenuStructure().GetDeveloperToolsMiscCategory())
+		.SetGroup(WorkspaceMenu::GetMenuStructure().GetDeveloperToolsDebugCategory())
 		.SetDisplayName(NSLOCTEXT("LogVisualizerApp", "TabTitle", "Visual Logger"))
 		.SetTooltipText(NSLOCTEXT("LogVisualizerApp", "TooltipText", "Opens Visual Logger tool."))
 		.SetIcon(FSlateIcon(FLogVisualizerStyle::GetStyleSetName(), "LogVisualizerApp.TabIcon"));
diff --git a/Engine/Source/Developer/ProjectLauncher/Private/ProjectLauncherModule.cpp b/Engine/Source/Developer/ProjectLauncher/Private/ProjectLauncherModule.cpp
--- a/Engine/Source/Developer/ProjectLauncher/Private/ProjectLauncherModule.cpp
+++ b/Engine/Source/Developer/ProjectLauncher/Private/ProjectLauncherModule.cpp
@@ -32,9 +32,6 @@ public:
 	
 	virtual void StartupModule() override
 	{
-#if WITH_EDITOR
-		FGlobalTabmanager::Get()->RegisterTabSpawner(ProjectLauncherTabName, FOnSpawnTab::CreateRaw(this, &FProjectLauncherModule::SpawnProjectLauncherTab));
-#else
 		// This is still experimental in the editor, so it'll be invoked specifically in FMainMenu if the experimental settings flag is set.
 		//@todo Enable this in the editor when no longer experimental
 		FGlobalTabmanager::Get()->RegisterNomadTabSpawner(ProjectLauncherTabName, FOnSpawnTab::CreateRaw(this, &FProjectLauncherModule::SpawnProjectLauncherTab))
@@ -42,16 +39,11 @@ public:
 			.SetTooltipText(NSLOCTEXT("FProjectLauncherModule", "ProjectLauncherTooltipText", "Open the Project Launcher tab."))
 			.SetIcon(FSlateIcon(FEditorStyle::GetStyleSetName(), "Launcher.TabIcon"))
 			.SetGroup(WorkspaceMenu::GetMenuStructure().GetToolsCategory());
-#endif
 	}
 
 	virtual void ShutdownModule() override
 	{
-#if WITH_EDITOR
-		FGlobalTabmanager::Get()->UnregisterTabSpawner(ProjectLauncherTabName);
-#else
 		FGlobalTabmanager::Get()->UnregisterNomadTabSpawner(ProjectLauncherTabName);
-#endif
 	}
 
 private:
diff --git a/Engine/Source/Developer/SessionFrontend/Private/SessionFrontendModule.cpp b/Engine/Source/Developer/SessionFrontend/Private/SessionFrontendModule.cpp
--- a/Engine/Source/Developer/SessionFrontend/Private/SessionFrontendModule.cpp
+++ b/Engine/Source/Developer/SessionFrontend/Private/SessionFrontendModule.cpp
@@ -62,11 +62,7 @@ public:
 			.SetTooltipText(NSLOCTEXT("FSessionFrontendModule", "FrontendTooltipText", "Open the Session Frontend tab."))
 			.SetIcon(FSlateIcon(FEditorStyle::GetStyleSetName(), "SessionFrontEnd.TabIcon"));
 
-#if WITH_EDITOR
-		TabSpawnerEntry.SetGroup(WorkspaceMenu::GetMenuStructure().GetDeveloperToolsMiscCategory());
-#else
 		TabSpawnerEntry.SetGroup(WorkspaceMenu::GetMenuStructure().GetToolsCategory());
-#endif
 	}
 
 	virtual void ShutdownModule() override
diff --git a/Engine/Source/Developer/SourceControl/Private/SSourceControlLogin.cpp b/Engine/Source/Developer/SourceControl/Private/SSourceControlLogin.cpp
--- a/Engine/Source/Developer/SourceControl/Private/SSourceControlLogin.cpp
+++ b/Engine/Source/Developer/SourceControl/Private/SSourceControlLogin.cpp
@@ -143,6 +143,7 @@ void SSourceControlLogin::Construct(const FArguments& InArgs)
 							SNew(SButton)
 							.VAlign(VAlign_Center)
 							.ButtonStyle(FAppStyle::Get(), "PrimaryButton")
+  							.TextStyle( FAppStyle::Get(), "DialogButtonText" )
 							.Text(LOCTEXT("AcceptSettings", "Accept Settings"))
 							.OnClicked( this, &SSourceControlLogin::OnAcceptSettings )
 							.IsEnabled( this, &SSourceControlLogin::IsAcceptSettingsEnabled )
@@ -153,6 +154,7 @@ void SSourceControlLogin::Construct(const FArguments& InArgs)
 						[
 							SNew(SButton)
 							.VAlign(VAlign_Center)
+  							.TextStyle( FAppStyle::Get(), "DialogButtonText" )
 							.Text(LOCTEXT("RunWithoutSourceControl", "Disable Source Control"))
 							.OnClicked( this, &SSourceControlLogin::OnDisableSourceControl )
 							.IsEnabled( this, &SSourceControlLogin::AreControlsEnabled )
diff --git a/Engine/Source/Editor/ClassViewer/Private/ClassViewerModule.cpp b/Engine/Source/Editor/ClassViewer/Private/ClassViewerModule.cpp
--- a/Engine/Source/Editor/ClassViewer/Private/ClassViewerModule.cpp
+++ b/Engine/Source/Editor/ClassViewer/Private/ClassViewerModule.cpp
@@ -47,7 +47,7 @@ void FClassViewerModule::StartupModule()
 	FGlobalTabmanager::Get()->RegisterNomadTabSpawner( ClassViewerModule::ClassViewerApp, FOnSpawnTab::CreateStatic( &CreateClassPickerTab ) )
 		.SetDisplayName( NSLOCTEXT("ClassViewerApp", "TabTitle", "Class Viewer") )
 		.SetTooltipText( NSLOCTEXT("ClassViewerApp", "TooltipText", "Displays all classes that exist within this project.") )
-		.SetGroup( WorkspaceMenu::GetMenuStructure().GetDeveloperToolsMiscCategory() )
+		.SetGroup( WorkspaceMenu::GetMenuStructure().GetToolsCategory() )
 		.SetIcon( FSlateIcon(FEditorStyle::GetStyleSetName(), "ClassViewer.TabIcon") );
 
 	ISettingsModule* SettingsModule = FModuleManager::GetModulePtr<ISettingsModule>("Settings");
diff --git a/Engine/Source/Editor/ContentBrowser/Private/ContentBrowserSingleton.cpp b/Engine/Source/Editor/ContentBrowser/Private/ContentBrowserSingleton.cpp
--- a/Engine/Source/Editor/ContentBrowser/Private/ContentBrowserSingleton.cpp
+++ b/Engine/Source/Editor/ContentBrowser/Private/ContentBrowserSingleton.cpp
@@ -45,7 +45,7 @@ FContentBrowserSingleton::FContentBrowserSingleton()
 	// Register the tab spawners for all content browsers
 	const FSlateIcon ContentBrowserIcon(FEditorStyle::GetStyleSetName(), "ContentBrowser.TabIcon");
 	const IWorkspaceMenuStructure& MenuStructure = WorkspaceMenu::GetMenuStructure();
-	TSharedRef<FWorkspaceItem> ContentBrowserGroup = MenuStructure.GetToolsCategory()->AddGroup(
+	TSharedRef<FWorkspaceItem> ContentBrowserGroup = MenuStructure.GetLevelEditorCategory()->AddGroup(
 		LOCTEXT( "WorkspaceMenu_ContentBrowserCategory", "Content Browser" ),
 		LOCTEXT( "ContentBrowserMenuTooltipText", "Open a Content Browser tab." ),
 		ContentBrowserIcon,
diff --git a/Engine/Source/Editor/DeviceProfileEditor/Private/DeviceProfileEditorModule.cpp b/Engine/Source/Editor/DeviceProfileEditor/Private/DeviceProfileEditorModule.cpp
--- a/Engine/Source/Editor/DeviceProfileEditor/Private/DeviceProfileEditorModule.cpp
+++ b/Engine/Source/Editor/DeviceProfileEditor/Private/DeviceProfileEditorModule.cpp
@@ -23,7 +23,7 @@ void FDeviceProfileEditorModule::StartupModule()
 	FGlobalTabmanager::Get()->RegisterNomadTabSpawner(DeviceProfileEditorName, FOnSpawnTab::CreateStatic(&FDeviceProfileEditorModule::SpawnDeviceProfileEditorTab))
 		.SetDisplayName( NSLOCTEXT("DeviceProfileEditor", "DeviceProfileEditorTitle", "Device Profiles") )
 		.SetIcon(FSlateIcon(FEditorStyle::GetStyleSetName(), "DeviceDetails.Tabs.ProfileEditor"))
-		.SetGroup( WorkspaceMenu::GetMenuStructure().GetDeveloperToolsMiscCategory() );
+		.SetGroup( WorkspaceMenu::GetMenuStructure().GetDeveloperToolsPlatformsCategory() );
 }
 
 
diff --git a/Engine/Source/Editor/EditorStyle/Private/StarshipStyle.cpp b/Engine/Source/Editor/EditorStyle/Private/StarshipStyle.cpp
--- a/Engine/Source/Editor/EditorStyle/Private/StarshipStyle.cpp
+++ b/Engine/Source/Editor/EditorStyle/Private/StarshipStyle.cpp
@@ -1356,13 +1356,10 @@ void FStarshipEditorStyle::FStyle::SetupGeneralStyles()
 	// About screen
 	if (IncludeEditorSpecificStyles())
 	{
-		Set( "AboutScreen.Background", new IMAGE_BRUSH( "About/Background", FVector2D(600,332), FLinearColor::White, ESlateBrushTileType::Both) );
-		Set( "AboutScreen.Facebook", new IMAGE_BRUSH( "About/FacebookIcon", FVector2D(35,35) ) );
-		Set( "AboutScreen.FacebookHovered", new IMAGE_BRUSH( "About/FacebookIcon_Hovered", FVector2D(35,35) ) );
-		Set( "AboutScreen.UE4", new IMAGE_BRUSH( "About/UE4Icon", FVector2D(50,50), FLinearColor::Gray) );
-		Set( "AboutScreen.UE4Hovered", new IMAGE_BRUSH( "About/UE4Icon", FVector2D(50,50), FLinearColor::White) );
-		Set( "AboutScreen.EpicGames", new IMAGE_BRUSH( "About/EpicGamesIcon", FVector2D(50,50) ) );
-		Set( "AboutScreen.EpicGamesHovered", new IMAGE_BRUSH( "About/EpicGamesIcon_Hovered", FVector2D(50,50) ) );
+		Set( "AboutScreen.Background", new IMAGE_BRUSH( "About/Background", FVector2D(688,317), FLinearColor::White, ESlateBrushTileType::Both) );
+		Set( "AboutScreen.UnrealLogo", new IMAGE_BRUSH_SVG( "About/UnrealLogo", FVector2D(40,40) ) );
+		Set( "AboutScreen.EpicGamesLogo", new IMAGE_BRUSH_SVG( "About/EpicGamesLogo", FVector2D(40,40) ) );
+		Set( "AboutScreen.TitleFont", DEFAULT_FONT( "Bold", 13) );
 	}
 #endif // WITH_EDITOR || (IS_PROGRAM && WITH_UNREAL_DEVELOPER_TOOLS)
 
@@ -4390,6 +4387,7 @@ void FStarshipEditorStyle::FStyle::SetupLevelEditorStyle()
 		{
 			Set( "LevelEditor.Tabs.Details",                new IMAGE_BRUSH_SVG("Starship/Common/Details", Icon16x16) );
 			Set( "LevelEditor.Tabs.Cinematics",             new IMAGE_BRUSH_SVG("Starship/Common/Cinematics", Icon16x16) );
+			Set( "LevelEditor.Tabs.VirtualProduction",      new IMAGE_BRUSH_SVG("Starship/Common/VirtualProduction", Icon16x16) );
 			Set( "LevelEditor.Tabs.EditorModes",            new IMAGE_BRUSH_SVG("Starship/Common/EditorModes", Icon16x16) );
 			Set( "LevelEditor.Tabs.Modes",                  new IMAGE_BRUSH_SVG("Starship/Common/EditorModes", Icon16x16) );
 			Set( "LevelEditor.Tabs.PlacementBrowser",       new IMAGE_BRUSH_SVG("Starship/Common/PlaceActors", Icon16x16) );
@@ -4416,6 +4414,7 @@ void FStarshipEditorStyle::FStyle::SetupLevelEditorStyle()
 
 			Set( "LevelEditor.Tabs.Viewports",              new IMAGE_BRUSH_SVG("Starship/Common/Viewports", Icon16x16) );
 			Set( "LevelEditor.Tabs.HLOD",                   new IMAGE_BRUSH_SVG("Starship/Common/HierarchicalLODOutliner", Icon16x16) );
+			Set( "LevelEditor.Tabs.Debug",                  new IMAGE_BRUSH_SVG("Starship/Common/Bug", Icon16x16) );
 		}
 #endif // WITH_EDITOR || (IS_PROGRAM && WITH_UNREAL_DEVELOPER_TOOLS)
 
diff --git a/Engine/Source/Editor/EnvironmentLightingViewer/Private/EnvironmentLightingModule.cpp b/Engine/Source/Editor/EnvironmentLightingViewer/Private/EnvironmentLightingModule.cpp
--- a/Engine/Source/Editor/EnvironmentLightingViewer/Private/EnvironmentLightingModule.cpp
+++ b/Engine/Source/Editor/EnvironmentLightingViewer/Private/EnvironmentLightingModule.cpp
@@ -44,7 +44,7 @@ void FEnvironmentLightingViewerModule::StartupModule()
 	FGlobalTabmanager::Get()->RegisterNomadTabSpawner(EnvironmentLightingViewerModule::EnvironmentLightingViewerApp, FOnSpawnTab::CreateStatic(&CreateEnvLightTab))//TODO picker tab
 		.SetDisplayName(NSLOCTEXT("EnvironmentLightingViewerApp", "TabTitle", "EnvironmentLighting Viewer"))
 		.SetTooltipText(NSLOCTEXT("EnvironmentLightingViewerApp", "TooltipText", "Environment lighting window."))
-		.SetGroup(WorkspaceMenu::GetMenuStructure().GetDeveloperToolsMiscCategory())
+		.SetGroup(WorkspaceMenu::GetMenuStructure().GetDeveloperToolsDebugCategory())
 		.SetIcon(FSlateIcon(FEditorStyle::GetStyleSetName(), "ClassViewer.TabIcon"));
 
 	// TODO setting module ISettingsModule* SettingsModule = FModuleManager::GetModulePtr<ISettingsModule>("Settings");
diff --git a/Engine/Source/Editor/LevelEditor/Private/LevelEditor.cpp b/Engine/Source/Editor/LevelEditor/Private/LevelEditor.cpp
--- a/Engine/Source/Editor/LevelEditor/Private/LevelEditor.cpp
+++ b/Engine/Source/Editor/LevelEditor/Private/LevelEditor.cpp
@@ -285,9 +285,9 @@ void FLevelEditorModule::StartupModule()
 	FGlobalTabmanager::Get()->RegisterTabSpawner("LevelEditor", FOnSpawnTab::CreateRaw( this, &FLevelEditorModule::SpawnLevelEditor ) )
 		.SetDisplayName( NSLOCTEXT("LevelEditor", "LevelEditorTab", "Level Editor") );
 
-	FModuleManager::LoadModuleChecked<ISlateReflectorModule>("SlateReflector").RegisterTabSpawner(MenuStructure.GetDeveloperToolsMiscCategory());
+	FModuleManager::LoadModuleChecked<ISlateReflectorModule>("SlateReflector").RegisterTabSpawner(MenuStructure.GetDeveloperToolsDebugCategory());
 
-	FModuleManager::LoadModuleChecked<FPixelInspectorModule>("PixelInspectorModule").RegisterTabSpawner(MenuStructure.GetDeveloperToolsMiscCategory());
+	FModuleManager::LoadModuleChecked<FPixelInspectorModule>("PixelInspectorModule").RegisterTabSpawner(MenuStructure.GetDeveloperToolsDebugCategory());
 
 	FMessageLogModule& MessageLogModule = FModuleManager::LoadModuleChecked<FMessageLogModule>("MessageLog");
 	MessageLogModule.RegisterLogListing("BuildAndSubmitErrors", LOCTEXT("BuildAndSubmitErrors", "Build and Submit Errors"));
diff --git a/Engine/Source/Editor/LevelEditor/Private/LevelEditorActions.cpp b/Engine/Source/Editor/LevelEditor/Private/LevelEditorActions.cpp
--- a/Engine/Source/Editor/LevelEditor/Private/LevelEditorActions.cpp
+++ b/Engine/Source/Editor/LevelEditor/Private/LevelEditorActions.cpp
@@ -3160,8 +3160,8 @@ void FLevelEditorCommands::RegisterCommands()
 	UI_COMMAND( BrowseViewportControls, "Viewport Controls...", "Opens the viewport controls cheat sheet", EUserInterfaceActionType::Button, FInputChord() );
 	UI_COMMAND( NewLevel, "New Level...", "Create a new level, or choose a level template to start from.", EUserInterfaceActionType::Button, FInputChord( EModifierKey::Control, EKeys::N ) );
 	UI_COMMAND( OpenLevel, "Open Level...", "Loads an existing level", EUserInterfaceActionType::Button, FInputChord( EModifierKey::Control, EKeys::O ) );
-	UI_COMMAND( Save, "Save Current", "Saves the current level to disk", EUserInterfaceActionType::Button, FInputChord(EModifierKey::Control, EKeys::S) );
-	UI_COMMAND( SaveAs, "Save Current As...", "Save the current level as...", EUserInterfaceActionType::Button, FInputChord( EModifierKey::Control|EModifierKey::Alt, EKeys::S ) );
+	UI_COMMAND( Save, "Save Current Level", "Saves the current level to disk", EUserInterfaceActionType::Button, FInputChord(EModifierKey::Control, EKeys::S) );
+	UI_COMMAND( SaveAs, "Save Current Level As...", "Save the current level as...", EUserInterfaceActionType::Button, FInputChord( EModifierKey::Control|EModifierKey::Alt, EKeys::S ) );
 	UI_COMMAND( SaveAllLevels, "Save All Levels", "Saves all unsaved levels to disk", EUserInterfaceActionType::Button, FInputChord() );
 	UI_COMMAND( ToggleFavorite, "Toggle Favorite", "Sets whether the currently loaded level will appear in the list of favorite levels", EUserInterfaceActionType::Button, FInputChord() );
 
diff --git a/Engine/Source/Editor/LevelEditor/Private/LevelEditorMenu.cpp b/Engine/Source/Editor/LevelEditor/Private/LevelEditorMenu.cpp
--- a/Engine/Source/Editor/LevelEditor/Private/LevelEditorMenu.cpp
+++ b/Engine/Source/Editor/LevelEditor/Private/LevelEditorMenu.cpp
@@ -30,36 +30,37 @@ void FLevelEditorMenu::RegisterLevelEditorMenus()
 		static void RegisterFileLoadAndSaveItems()
 		{
 			UToolMenu* Menu = UToolMenus::Get()->ExtendMenu("LevelEditor.MainMenu.File");
-			FToolMenuSection& Section = Menu->FindOrAddSection("FileLoadAndSave");
+
+			FToolMenuSection& OpenSection = Menu->FindOrAddSection("FileOpen");
+
 			FToolMenuInsert InsertPos(NAME_None, EToolMenuInsertType::First);
 
 			// New Level
-			Section.AddMenuEntry( FLevelEditorCommands::Get().NewLevel ).InsertPosition = InsertPos;
+			OpenSection.AddMenuEntry( FLevelEditorCommands::Get().NewLevel ).InsertPosition = InsertPos;
 
 			// Open Level
-			Section.AddMenuEntry( FLevelEditorCommands::Get().OpenLevel ).InsertPosition = InsertPos;
+			OpenSection.AddMenuEntry( FLevelEditorCommands::Get().OpenLevel ).InsertPosition = InsertPos;
 
 			// Open Asset
 			//@TODO: Doesn't work when summoned from here: Section.AddMenuEntry( FGlobalEditorCommonCommands::Get().SummonOpenAssetDialog );
 
+			FToolMenuSection& SaveSection = Menu->FindOrAddSection("FileSave");
+
 			// Save
-			Section.AddMenuEntry( FLevelEditorCommands::Get().Save ).InsertPosition = InsertPos;
+			SaveSection.AddMenuEntry( FLevelEditorCommands::Get().Save ).InsertPosition = InsertPos;
 	
 			// Save As
-			Section.AddMenuEntry( FLevelEditorCommands::Get().SaveAs ).InsertPosition = InsertPos;
-
-			// Save Levels
-			Section.AddMenuEntry( FLevelEditorCommands::Get().SaveAllLevels ).InsertPosition = InsertPos;
+			SaveSection.AddMenuEntry( FLevelEditorCommands::Get().SaveAs ).InsertPosition = InsertPos;
 		}
 
 		static void FillFileRecentAndFavoriteFileItems()
 		{
 			UToolMenu* Menu = UToolMenus::Get()->ExtendMenu("LevelEditor.MainMenu.File");
-			FToolMenuInsert SectionInsertPos("FileRecentFiles", EToolMenuInsertType::Before);
+			FToolMenuInsert SectionInsertPos("FileSave", EToolMenuInsertType::After);
 
 			// Import/Export
 			{
-				FToolMenuSection& Section = Menu->AddSection("FileActors", LOCTEXT("ImportExportHeading", "Actors"), SectionInsertPos);
+				FToolMenuSection& Section = Menu->AddSection("FileActors", LOCTEXT("ImportExportHeading", "Import/Export"), SectionInsertPos);
 				{
 					// Import Into Level
 					Section.AddMenuEntry(FLevelEditorCommands::Get().ImportScene);
@@ -121,7 +122,7 @@ void FLevelEditorMenu::RegisterLevelEditorMenus()
 					}
 				};
 
-				FToolMenuSection& Section = Menu->AddSection("FileFavoriteLevels", TAttribute<FText>(), SectionInsertPos);
+				FToolMenuSection& Section = Menu->FindOrAddSection("FileOpen");
 
 				Section.AddDynamicEntry("FileFavoriteLevels", FNewToolMenuSectionDelegate::CreateLambda([](FToolMenuSection& InSection)
 				{
@@ -165,7 +166,7 @@ void FLevelEditorMenu::RegisterLevelEditorMenus()
 					}
 				};
 
-				FToolMenuSection& Section = Menu->AddSection("FileRecentLevels", TAttribute<FText>(), SectionInsertPos);
+				FToolMenuSection& Section = Menu->FindOrAddSection("FileOpen");
 				Section.AddDynamicEntry("FileRecentLevels", FNewToolMenuSectionDelegate::CreateLambda([](FToolMenuSection& InSection)
 				{
 					IMainFrameModule& MainFrameModule = FModuleManager::LoadModuleChecked<IMainFrameModule>( "MainFrame" );
@@ -227,6 +228,7 @@ void FLevelEditorMenu::RegisterLevelEditorMenus()
 
 			FToolMenuSection& Section = Menu->FindOrAddSection(NAME_None);
 
+
 			FToolMenuEntry& Entry =
 				Section.AddSubMenu(
 					"Build",
@@ -237,6 +239,12 @@ void FLevelEditorMenu::RegisterLevelEditorMenus()
 
 			Entry.InsertPosition = FToolMenuInsert("Help", EToolMenuInsertType::Before);
 		}
+
+
+		static void ExtendToolsMenu()
+		{
+			UToolMenu* Menu = UToolMenus::Get()->RegisterMenu("LevelEditor.MainMenu.Tools", "MainFrame.MainMenu.Tools");
+		}
 	};
 
 	UToolMenus* ToolMenus = UToolMenus::Get();
@@ -244,6 +252,7 @@ void FLevelEditorMenu::RegisterLevelEditorMenus()
 	ToolMenus->RegisterMenu("LevelEditor.MainMenu.File", "MainFrame.MainTabMenu.File");
 	ToolMenus->RegisterMenu("LevelEditor.MainMenu.Window", "MainFrame.MainMenu.Window");
 
+	// Add other top level menus
 	Local::ExtendMenuBar();
 
 	// Add level loading and saving menu items
@@ -258,6 +267,10 @@ void FLevelEditorMenu::RegisterLevelEditorMenus()
 	// Extend the Help menu
 	Local::ExtendHelpMenu();
 
+	// Extend the Tools menu
+	Local::ExtendToolsMenu();
+
+	// Extend the Build menu
 	RegisterBuildMenu();
 
 }
diff --git a/Engine/Source/Editor/LevelEditor/Private/SLevelEditor.cpp b/Engine/Source/Editor/LevelEditor/Private/SLevelEditor.cpp
--- a/Engine/Source/Editor/LevelEditor/Private/SLevelEditor.cpp
+++ b/Engine/Source/Editor/LevelEditor/Private/SLevelEditor.cpp
@@ -1372,7 +1372,7 @@ TSharedRef<SWidget> SLevelEditor::RestoreContentArea( const TSharedRef<SDockTab>
 			LevelEditorTabManager->RegisterTabSpawner(LevelEditorTabIds::LevelEditorStatsViewer, FOnSpawnTab::CreateSP<SLevelEditor, FName, FString>(this, &SLevelEditor::SpawnLevelEditorTab, LevelEditorTabIds::LevelEditorStatsViewer, FString()))
 				.SetDisplayName(NSLOCTEXT("LevelEditorTabs", "LevelEditorStatsViewer", "Statistics"))
 				.SetTooltipText(NSLOCTEXT("LevelEditorTabs", "LevelEditorStatsViewerTooltipText", "Open the Statistics tab, in order to see data pertaining to lighting, textures and primitives."))
-				.SetGroup(MenuStructure.GetLevelEditorCategory())
+				.SetGroup(MenuStructure.GetDeveloperToolsAuditCategory())
 				.SetIcon(StatsViewerIcon);
 		}
 
diff --git a/Engine/Source/Editor/MainFrame/Private/Menus/MainMenu.cpp b/Engine/Source/Editor/MainFrame/Private/Menus/MainMenu.cpp
--- a/Engine/Source/Editor/MainFrame/Private/Menus/MainMenu.cpp
+++ b/Engine/Source/Editor/MainFrame/Private/Menus/MainMenu.cpp
@@ -40,39 +40,25 @@ void FMainMenu::RegisterFileMenu()
 	UToolMenus* ToolMenus = UToolMenus::Get();
 	UToolMenu* FileMenu = ToolMenus->RegisterMenu("MainFrame.MainMenu.File");
 
-	FToolMenuSection& FileLoadAndSaveSection = FileMenu->AddSection("FileLoadAndSave", LOCTEXT("LoadSandSaveHeading", "Load and Save"));
+	FToolMenuSection& FileLoadAndSaveSection = FileMenu->AddSection("FileOpen", LOCTEXT("FileOpenHeading", "Open"), FToolMenuInsert(NAME_None, EToolMenuInsertType::First));
 	{
 		// Open Asset...
 		FileLoadAndSaveSection.AddMenuEntry(FGlobalEditorCommonCommands::Get().SummonOpenAssetDialog);
+	}
+
+
+	FToolMenuSection& FileSaveSection = FileMenu->AddSection("FileSave", LOCTEXT("FileSaveHeading", "Save"), FToolMenuInsert("FileOpen", EToolMenuInsertType::After));
+	{
 
 		// Save All
-		FileLoadAndSaveSection.AddMenuEntry(FMainFrameCommands::Get().SaveAll);
+		FileSaveSection.AddMenuEntry(FMainFrameCommands::Get().SaveAll);
 
 		// Choose specific files to save
-		FileLoadAndSaveSection.AddMenuEntry(FMainFrameCommands::Get().ChooseFilesToSave);
-
-		FileLoadAndSaveSection.AddDynamicEntry("SourceControl", FNewToolMenuSectionDelegate::CreateLambda([](FToolMenuSection& InSection)
-		{
-			if (ISourceControlModule::Get().IsEnabled() && ISourceControlModule::Get().GetProvider().IsAvailable())
-			{
-				// Choose specific files to submit
-				InSection.AddMenuEntry(FMainFrameCommands::Get().ViewChangelists, 
-					TAttribute<FText>(),
-					TAttribute<FText>(),
-					FSlateIcon(FEditorStyle::GetStyleSetName(), "SourceControl.ChangelistsTab"));
-			}
-			else
-			{
-				InSection.AddMenuEntry(FMainFrameCommands::Get().ConnectToSourceControl, 
-					TAttribute<FText>(),
-					TAttribute<FText>(),
-					FSlateIcon(FEditorStyle::GetStyleSetName(), "SourceControl.Actions.Connect"));
-			}
-		}));
+		FileSaveSection.AddMenuEntry(FMainFrameCommands::Get().ChooseFilesToSave);
 	}
 
 	RegisterFileProjectMenu();
-	RegisterRecentFileAndExitMenuItems();
+	RegisterExitMenuItems();
 }
 
 #undef LOCTEXT_NAMESPACE
@@ -171,6 +157,25 @@ void FMainMenu::RegisterEditMenu()
 				FGlobalTabmanager::Get()->PopulateTabSpawnerMenu(InBuilder, "PluginsEditor");
 			}
 		}));
+
+		Section.AddDynamicEntry("SourceControlDyn", FNewToolMenuSectionDelegate::CreateLambda([](FToolMenuSection& InSection)
+		{
+			if (ISourceControlModule::Get().IsEnabled() && ISourceControlModule::Get().GetProvider().IsAvailable())
+			{
+				// Choose specific files to submit
+				InSection.AddMenuEntry(FMainFrameCommands::Get().ViewChangelists, 
+					TAttribute<FText>(),
+					TAttribute<FText>(),
+					FSlateIcon(FEditorStyle::GetStyleSetName(), "SourceControl.ChangelistsTab"));
+			}
+			else
+			{
+				InSection.AddMenuEntry(FMainFrameCommands::Get().ConnectToSourceControl, 
+					TAttribute<FText>(),
+					TAttribute<FText>(),
+					FSlateIcon(FEditorStyle::GetStyleSetName(), "SourceControl.Actions.Connect"));
+			}
+		}));
 	}
 }
 
@@ -201,17 +206,7 @@ void FMainMenu::RegisterWindowMenu()
 		}
 	}));
 
-	// Project Launcher section
-	{
-		FToolMenuSection& Section = Menu->AddSection("WindowGlobalTabSpawners");
-		Section.AddMenuEntry(
-			"ProjectLauncher",
-			LOCTEXT("ProjectLauncherLabel", "Project Launcher"),
-			LOCTEXT("ProjectLauncherToolTip", "The Project Launcher provides advanced workflows for packaging, deploying and launching your projects."),
-			FSlateIcon(FEditorStyle::GetStyleSetName(), "Launcher.TabIcon"),
-			FUIAction(FExecuteAction::CreateStatic(&FMainMenu::OpenProjectLauncher))
-			);
-	}
+	
 
 	// Experimental section
 	{
@@ -385,6 +380,7 @@ void FMainMenu::RegisterMainMenu()
 	RegisterFileMenu();
 	RegisterEditMenu();
 	RegisterWindowMenu();
+	RegisterToolsMenu();
 	RegisterHelpMenu();
 
 	UToolMenu* MenuBar = ToolMenus->RegisterMenu(MainMenuName, NAME_None, EMultiBoxType::MenuBar);
@@ -417,6 +413,15 @@ void FMainMenu::RegisterMainMenu()
 		LOCTEXT("WindowMenu_ToolTip", "Open new windows or tabs.")
 		);
 
+	MenuBar->AddSubMenu(
+		"MainMenu",
+		NAME_None,
+		"Tools",
+		LOCTEXT("ToolsMenu", "Tools"),
+		LOCTEXT("ToolsMenu_ToolTip", "Level Tools")
+		);
+
+
 	MenuBar->AddSubMenu(
 		"MainMenu",
 		NAME_None,
@@ -440,20 +445,10 @@ void FMainMenu::RegisterFileProjectMenu()
 
 	UToolMenus* ToolMenus = UToolMenus::Get();
 	UToolMenu* MainTabFileMenu = ToolMenus->ExtendMenu("MainFrame.MainTabMenu.File");
-	FToolMenuSection& Section = MainTabFileMenu->AddSection("FileProject", LOCTEXT("ProjectHeading", "Project"), FToolMenuInsert("FileLoadAndSave", EToolMenuInsertType::After));
+	FToolMenuSection& Section = MainTabFileMenu->AddSection("FileProject", LOCTEXT("ProjectHeading", "Project"));
 
 	Section.AddMenuEntry( FMainFrameCommands::Get().NewProject );
 	Section.AddMenuEntry( FMainFrameCommands::Get().OpenProject );
-
-	FText ShortIDEName = FSourceCodeNavigation::GetSelectedSourceCodeIDE();
-	FSlateIcon OpenIDEIcon = FSourceCodeNavigation::GetOpenSourceCodeIDEIcon();
-	FSlateIcon RefreshIDEIcon = FSourceCodeNavigation::GetRefreshSourceCodeIDEIcon(); 
-
-	Section.AddMenuEntry( FMainFrameCommands::Get().AddCodeToProject,
-		TAttribute<FText>(),
-		FText::Format(LOCTEXT("AddCodeToProjectTooltip", "Adds C++ code to the project. The code can only be compiled if you have {0} installed."), ShortIDEName)
-	);
-
 	/*
 	MenuBuilder.AddMenuEntry( FMainFrameCommands::Get().LocalizeProject,
 		NAME_None,
@@ -469,6 +464,38 @@ void FMainMenu::RegisterFileProjectMenu()
 	);
 	*/
 
+
+	Section.AddMenuEntry(FMainFrameCommands::Get().ZipUpProject);
+
+	if (GetDefault<UEditorStyleSettings>()->bShowProjectMenus && FMainFrameActionCallbacks::RecentProjects.Num() > 0)
+	{
+		Section.AddSubMenu(
+			"RecentProjects",
+			LOCTEXT("SwitchProjectSubMenu", "Recent Projects"),
+			LOCTEXT("SwitchProjectSubMenu_ToolTip", "Select a project to switch to"),
+			FNewToolMenuDelegate::CreateStatic(&FRecentProjectsMenu::MakeMenu),
+			false,
+			FSlateIcon(FEditorStyle::GetStyleSetName(), "MainFrame.RecentProjects")
+		);
+	}
+
+}
+
+void FMainMenu::RegisterToolsMenu()
+{
+	UToolMenu* Menu = UToolMenus::Get()->RegisterMenu("MainFrame.MainMenu.Tools");
+	FToolMenuSection& Section = Menu->AddSection("Programming", LOCTEXT("ProgrammingHeading", "Programming"));
+
+
+	FText ShortIDEName = FSourceCodeNavigation::GetSelectedSourceCodeIDE();
+	FSlateIcon OpenIDEIcon = FSourceCodeNavigation::GetOpenSourceCodeIDEIcon();
+	FSlateIcon RefreshIDEIcon = FSourceCodeNavigation::GetRefreshSourceCodeIDEIcon(); 
+
+	Section.AddMenuEntry( FMainFrameCommands::Get().AddCodeToProject,
+		TAttribute<FText>(),
+		FText::Format(LOCTEXT("AddCodeToProjectTooltip", "Adds C++ code to the project. The code can only be compiled if you have {0} installed."), ShortIDEName)
+	);
+
 	Section.AddDynamicEntry("CodeProject", FNewToolMenuSectionDelegate::CreateLambda([ShortIDEName, RefreshIDEIcon](FToolMenuSection& InSection)
 	{
 		if (FSourceCodeNavigation::DoesModuleSolutionExist())
@@ -495,32 +522,40 @@ void FMainMenu::RegisterFileProjectMenu()
 		OpenIDEIcon
 	);
 
-	Section.AddMenuEntry(FMainFrameCommands::Get().ZipUpProject);
+
+	// Level Editor, General, and Testing sections
+	// Automatically populate tab spawners from TabManager
+	Menu->AddDynamicSection("TabManagerSection", FNewToolMenuDelegateLegacy::CreateLambda([](FMenuBuilder& InBuilder, UToolMenu* InData)
+	{
+		if (USlateTabManagerContext* TabManagerContext = InData->FindContext<USlateTabManagerContext>())
+		{
+			TSharedPtr<FTabManager> TabManager = TabManagerContext->TabManager.Pin();
+			if (TabManager.IsValid())
+			{
+				// Local editor tabs
+				TabManager->PopulateLocalTabSpawnerMenu(InBuilder);
+
+				// General tabs
+				const IWorkspaceMenuStructure& MenuStructure = WorkspaceMenu::GetMenuStructure();
+				TabManager->PopulateTabSpawnerMenu(InBuilder, MenuStructure.GetToolsStructureRoot());
+			}
+		}
+	}));
+
 }
 
-void FMainMenu::RegisterRecentFileAndExitMenuItems()
+void FMainMenu::RegisterExitMenuItems()
 {
 	UToolMenus* ToolMenus = UToolMenus::Get();
+
 	UToolMenu* MainTabFileMenu = ToolMenus->RegisterMenu("MainFrame.MainTabMenu.File", "MainFrame.MainMenu.File");
 
-	{
-		FToolMenuSection& Section = MainTabFileMenu->AddSection("FileRecentFiles");
-		if (GetDefault<UEditorStyleSettings>()->bShowProjectMenus && FMainFrameActionCallbacks::RecentProjects.Num() > 0)
-		{
-			Section.AddSubMenu(
-				"RecentProjects",
-				LOCTEXT("SwitchProjectSubMenu", "Recent Projects"),
-				LOCTEXT("SwitchProjectSubMenu_ToolTip", "Select a project to switch to"),
-				FNewToolMenuDelegate::CreateStatic(&FRecentProjectsMenu::MakeMenu),
-				false,
-				FSlateIcon(FEditorStyle::GetStyleSetName(), "MainFrame.RecentProjects")
-			);
-		}
-	}
+	
 
 #if !PLATFORM_MAC // Handled by app's menu in menu bar
 	{
-		FToolMenuSection& Section = MainTabFileMenu->AddSection("Exit");
+		FToolMenuSection& Section = MainTabFileMenu->AddSection("Exit", LOCTEXT("Exit", "Exit"), FToolMenuInsert("FileProject", EToolMenuInsertType::After));
+		Section.AddSeparator("Exit");
 		Section.AddMenuEntry( FMainFrameCommands::Get().Exit );
 	}
 #endif
diff --git a/Engine/Source/Editor/MainFrame/Private/Menus/MainMenu.h b/Engine/Source/Editor/MainFrame/Private/Menus/MainMenu.h
--- a/Engine/Source/Editor/MainFrame/Private/Menus/MainMenu.h
+++ b/Engine/Source/Editor/MainFrame/Private/Menus/MainMenu.h
@@ -37,6 +37,7 @@ public:
 	 * @param MenuName Identifier associated with the menu.
 	 * @return ToolMenuContext Context containing state.
 	 */
+	UE_DEPRECATED(4.26, "MakeMainTabMenu has been deprecated, use MakeMainMenu instead.")
 	static TSharedRef<SWidget> MakeMainTabMenu( const TSharedPtr<FTabManager>& TabManager, const FName MenuName, FToolMenuContext& ToolMenuContext );
 
 	/**
@@ -62,6 +63,11 @@ protected:
 	 */
 	static void RegisterWindowMenu();
 
+	/**
+	 * Called to fill the app menu's content.
+	 */
+	static void RegisterToolsMenu();
+
 	/**
 	 * Called to fill the help menu's content.
 	 */
@@ -75,7 +81,7 @@ protected:
 	/**
 	 * Called to fill the file menu's recent and exit content.
 	 */
-	static void RegisterRecentFileAndExitMenuItems();
+	static void RegisterExitMenuItems();
 
 private:
 	/** 
diff --git a/Engine/Source/Editor/MergeActors/Private/MergeActorsModule.cpp b/Engine/Source/Editor/MergeActors/Private/MergeActorsModule.cpp
--- a/Engine/Source/Editor/MergeActors/Private/MergeActorsModule.cpp
+++ b/Engine/Source/Editor/MergeActors/Private/MergeActorsModule.cpp
@@ -107,7 +107,7 @@ void FMergeActorsModule::StartupModule()
 	FGlobalTabmanager::Get()->RegisterNomadTabSpawner(MergeActorsTabName, FOnSpawnTab::CreateRaw(this, &FMergeActorsModule::CreateMergeActorsTab))
 		.SetDisplayName(NSLOCTEXT("MergeActorsModule", "TabTitle", "Merge Actors"))
 		.SetTooltipText(NSLOCTEXT("MergeActorsModule", "TooltipText", "Open the Merge Actors tab."))
-		.SetGroup(WorkspaceMenu::GetMenuStructure().GetDeveloperToolsMiscCategory())
+		.SetGroup(WorkspaceMenu::GetMenuStructure().GetToolsCategory())
 		.SetIcon(FSlateIcon(FEditorStyle::GetStyleSetName(), "MergeActors.TabIcon"));
 
 	// Register built-in merging tools straight away
diff --git a/Engine/Source/Editor/SourceControlWindows/Private/SourceControlWindowsModule.cpp b/Engine/Source/Editor/SourceControlWindows/Private/SourceControlWindowsModule.cpp
--- a/Engine/Source/Editor/SourceControlWindows/Private/SourceControlWindowsModule.cpp
+++ b/Engine/Source/Editor/SourceControlWindows/Private/SourceControlWindowsModule.cpp
@@ -59,17 +59,11 @@ void FSourceControlWindowsModule::StartupModule()
 
 	// Create a Source Control group under the Tools category
 	const FSlateIcon SourceControlIcon(FEditorStyle::GetStyleSetName(), "SourceControl.ChangelistsTab");
-	const IWorkspaceMenuStructure& MenuStructure = WorkspaceMenu::GetMenuStructure();
-	TSharedRef<FWorkspaceItem> SourceControlGroup = MenuStructure.GetToolsCategory()->AddGroup(
-		LOCTEXT("WorkspaceMenu_SourceControlCategory", "Source Control"),
-		LOCTEXT("SourceControlMenuTooltipText", "Source Control Operations"),
-		SourceControlIcon,
-		true);
 
 	// Register the changlist tab spawner
 	FGlobalTabmanager::Get()->RegisterNomadTabSpawner(SourceControlChangelistsTabName, FOnSpawnTab::CreateRaw(this, &FSourceControlWindowsModule::CreateChangelistsTab))
 		.SetDisplayName(LOCTEXT("ChangelistsTabTitle", "Changelists"))
-		.SetGroup(SourceControlGroup)
+		.SetGroup(WorkspaceMenu::GetMenuStructure().GetToolsCategory())
 		.SetIcon(SourceControlIcon);
 
 #if WITH_RELOAD
diff --git a/Engine/Source/Editor/StructViewer/Private/StructViewerModule.cpp b/Engine/Source/Editor/StructViewer/Private/StructViewerModule.cpp
--- a/Engine/Source/Editor/StructViewer/Private/StructViewerModule.cpp
+++ b/Engine/Source/Editor/StructViewer/Private/StructViewerModule.cpp
@@ -46,7 +46,7 @@ void FStructViewerModule::StartupModule()
 	FGlobalTabmanager::Get()->RegisterNomadTabSpawner(StructViewerModule::StructViewerApp, FOnSpawnTab::CreateStatic(&CreateStructPickerTab))
 		.SetDisplayName(NSLOCTEXT("StructViewerApp", "TabTitle", "Struct Viewer"))
 		.SetTooltipText(NSLOCTEXT("StructViewerApp", "TooltipText", "Displays all structs that exist within this project."))
-		.SetGroup(WorkspaceMenu::GetMenuStructure().GetDeveloperToolsMiscCategory())
+		.SetGroup(WorkspaceMenu::GetMenuStructure().GetToolsCategory())
 		.SetIcon(FSlateIcon(FEditorStyle::GetStyleSetName(), "ClassViewer.TabIcon"));
 
 	ISettingsModule* SettingsModule = FModuleManager::GetModulePtr<ISettingsModule>("Settings");
diff --git a/Engine/Source/Editor/TimeManagementEditor/Private/STimecodeProviderTab.cpp b/Engine/Source/Editor/TimeManagementEditor/Private/STimecodeProviderTab.cpp
--- a/Engine/Source/Editor/TimeManagementEditor/Private/STimecodeProviderTab.cpp
+++ b/Engine/Source/Editor/TimeManagementEditor/Private/STimecodeProviderTab.cpp
@@ -48,7 +48,7 @@ void STimecodeProviderTab::RegisterNomadTabSpawner()
 		LevelEditorTabManager->RegisterTabSpawner(TimecodeProviderTab::NAME_TimecodeProviderTab, FOnSpawnTab::CreateStatic(&TimecodeProviderTab::CreateTab))
 			.SetDisplayName(NSLOCTEXT("TimecodeProviderTab", "DisplayName", "Timecode Provider"))
 			.SetTooltipText(NSLOCTEXT("TimecodeProviderTab", "TooltipText", "Displays the Timecode and the state of the current Timecode Provider."))
-			.SetGroup(WorkspaceMenu::GetMenuStructure().GetDeveloperToolsMiscCategory())
+			.SetGroup(WorkspaceMenu::GetMenuStructure().GetLevelEditorVirtualProductionCategory())
 			.SetIcon(FSlateIcon(FEditorStyle::GetStyleSetName(), "TimecodeProvider.TabIcon"));
 	};
 
diff --git a/Engine/Source/Editor/WorkspaceMenuStructure/Private/WorkspaceMenuStructureModule.cpp b/Engine/Source/Editor/WorkspaceMenuStructure/Private/WorkspaceMenuStructureModule.cpp
--- a/Engine/Source/Editor/WorkspaceMenuStructure/Private/WorkspaceMenuStructureModule.cpp
+++ b/Engine/Source/Editor/WorkspaceMenuStructure/Private/WorkspaceMenuStructureModule.cpp
@@ -20,6 +20,10 @@ public:
 	{
 		return MenuRoot.ToSharedRef();
 	}
+	virtual TSharedRef<FWorkspaceItem> GetToolsStructureRoot() const override
+	{
+		return ToolsMenuRoot.ToSharedRef();
+	}
 
 	virtual TSharedRef<FWorkspaceItem> GetLevelEditorCategory() const override
 	{
@@ -46,6 +50,11 @@ public:
 		return LevelEditorCinematicsCategory.ToSharedRef();
 	}
 
+	virtual TSharedRef<FWorkspaceItem> GetLevelEditorVirtualProductionCategory() const override
+	{
+		return LevelEditorVirtualProductionCategory.ToSharedRef();
+	}
+
 	virtual TSharedRef<FWorkspaceItem> GetToolsCategory() const override
 	{
 		return ToolsCategory.ToSharedRef();
@@ -66,6 +75,16 @@ public:
 		return DeveloperToolsProfilingCategory.ToSharedRef();
 	}
 
+	virtual TSharedRef<FWorkspaceItem> GetDeveloperToolsAuditCategory() const override
+	{
+		return DeveloperToolsAuditCategory.ToSharedRef();
+	}	
+
+	virtual TSharedRef<FWorkspaceItem> GetDeveloperToolsPlatformsCategory() const override
+	{
+		return DeveloperToolsPlatformsCategory.ToSharedRef();
+	}
+
 	virtual TSharedRef<FWorkspaceItem> GetDeveloperToolsMiscCategory() const override
 	{
 		return DeveloperToolsMiscCategory.ToSharedRef();
@@ -87,33 +106,38 @@ public:
 		LevelEditorViewportsCategory = LevelEditorCategory->AddGroup(LOCTEXT( "WorkspaceMenu_LevelEditorViewportCategory", "Viewports" ), LOCTEXT( "WorkspaceMenu_LevelEditorViewportCategoryTooltip", "Open a Viewport tab." ), FSlateIcon(FEditorStyle::GetStyleSetName(), "LevelEditor.Tabs.Viewports"), true);
 		LevelEditorDetailsCategory = LevelEditorCategory->AddGroup(LOCTEXT("WorkspaceMenu_LevelEditorDetailCategory", "Details" ), LOCTEXT("WorkspaceMenu_LevelEditorDetailCategoryTooltip", "Open a Details tab." ), FSlateIcon(FEditorStyle::GetStyleSetName(), "LevelEditor.Tabs.Details"), true );
 		LevelEditorCinematicsCategory = LevelEditorCategory->AddGroup(LOCTEXT("WorkspaceMenu_LevelEditorCinematicsCategory", "Cinematics"), LOCTEXT("WorkspaceMenu_LevelEditorCinematicsCategoryTooltip", "Open a Cinematics tab."), FSlateIcon(FEditorStyle::GetStyleSetName(), "LevelEditor.Tabs.Cinematics"), true);
+		LevelEditorVirtualProductionCategory = LevelEditorCategory->AddGroup(LOCTEXT("WorkspaceMenu_LevelEditorVirtualProductionCategory", "Virtual Production"), LOCTEXT("WorkspaceMenu_LevelEditorVirtualProductionCategoryTooltip", "Open a Virtual Production tab."), FSlateIcon(FEditorStyle::GetStyleSetName(), "LevelEditor.Tabs.VirtualProduction"), true);
 		LevelEditorModesCategory = LevelEditorCategory->AddGroup(LOCTEXT("WorkspaceMenu_LevelEditorToolsCategory", "Editor Modes" ), FSlateIcon(FEditorStyle::GetStyleSetName(), "LevelEditor.Tabs.EditorModes"), true );
+
 	}
 
 	void ResetToolsCategory()
 	{
-		ToolsCategory->ClearItems();
+		// Developer tools sections
+		DeveloperToolsCategory->ClearItems();
+		DeveloperToolsSubMenuDebugCategory =     DeveloperToolsCategory->AddGroup(LOCTEXT("WorkspaceMenu_DeveloperToolsDebugCategory", "Debug"), FSlateIcon(FEditorStyle::GetStyleSetName(), "LevelEditor.Tabs.Debug"), true);
 
-		// Developer tools sub menu
-		DeveloperToolsCategory = ToolsCategory->AddGroup(LOCTEXT("WorkspaceMenu_DeveloperToolsCategory", "Developer Tools"), FSlateIcon(FEditorStyle::GetStyleSetName(), "DeveloperTools.MenuIcon"));
+		DeveloperToolsDebugCategory = DeveloperToolsSubMenuDebugCategory->AddGroup(LOCTEXT("WorkspaceMenu_DeveloperToolsDebugCategory", "Debug"), FSlateIcon(), true);
+		AutomationToolsCategory = DeveloperToolsSubMenuDebugCategory->AddGroup(LOCTEXT("WorkspaceMenu_AutomationToolsCategory", "Testing"), FSlateIcon(), true);
 
-		// Developer tools sections
-		DeveloperToolsDebugCategory = DeveloperToolsCategory->AddGroup(LOCTEXT("WorkspaceMenu_DeveloperToolsDebugCategory", "Debug"), FSlateIcon(), true);
-		DeveloperToolsLogCategory = DeveloperToolsCategory->AddGroup(LOCTEXT("WorkspaceMenu_DeveloperToolsLogCategory", "Log"), FSlateIcon(), true);
-		DeveloperToolsProfilingCategory = DeveloperToolsCategory->AddGroup(LOCTEXT("WorkspaceMenu_DeveloperToolsProfilingCategory", "Profiling"), FSlateIcon(), true);
-		DeveloperToolsMiscCategory = DeveloperToolsCategory->AddGroup(LOCTEXT("WorkspaceMenu_DeveloperToolsMiscCategory", "Miscellaneous"), FSlateIcon(), true);
-		
-		// Automation tools sub menu
-		AutomationToolsCategory = FWorkspaceItem::NewGroup(LOCTEXT("WorkspaceMenu_AutomationToolsCategory", "Automation Tools"), FSlateIcon(), true);
+		DeveloperToolsProfilingCategory = DeveloperToolsCategory->AddGroup(LOCTEXT("WorkspaceMenu_DeveloperToolsProfilingCategory", "Profile"),  FSlateIcon(FEditorStyle::GetStyleSetName(), "LevelEditor.Tabs.Debug"), true);
+		DeveloperToolsAuditCategory =     DeveloperToolsCategory->AddGroup(LOCTEXT("WorkspaceMenu_DeveloperToolsAuditCategory", "Audit"), FSlateIcon(FEditorStyle::GetStyleSetName(), "LevelEditor.Tabs.Debug"), true);
+		DeveloperToolsPlatformsCategory = DeveloperToolsCategory->AddGroup(LOCTEXT("WorkspaceMenu_DeveloperToolsPlatforms:WCategory", "Platforms"), FSlateIcon(FEditorStyle::GetStyleSetName(), "LevelEditor.Tabs.Debug"), true);
+		DeveloperToolsMiscCategory =      DeveloperToolsCategory->AddGroup(LOCTEXT("WorkspaceMenu_DeveloperToolsMiscCategory", "Miscellaneous"), FSlateIcon(FEditorStyle::GetStyleSetName(), "LevelEditor.Tabs.Debug"), true);
 	}
 
 public:
 	FWorkspaceMenuStructure()
 		: MenuRoot ( FWorkspaceItem::NewGroup(LOCTEXT( "WorkspaceMenu_Root", "Menu Root" )) )
+		, ToolsMenuRoot ( FWorkspaceItem::NewGroup(LOCTEXT( "WorkspaceToolsMenu_Root", "Tools Menu Root" )) )
 		, LevelEditorCategory ( MenuRoot->AddGroup(LOCTEXT( "WorkspaceMenu_LevelEditorCategory", "Level Editor" ), FSlateIcon(), true) )
-		, ToolsCategory ( MenuRoot->AddGroup(LOCTEXT( "WorkspaceMenu_ToolsCategory", "General" ), FSlateIcon(), true) )
+		, ToolsCategory ( ToolsMenuRoot->AddGroup(LOCTEXT( "WorkspaceMenu_ToolsCategory", "Tools" ), FSlateIcon(), true) )
+		, DeveloperToolsCategory ( ToolsMenuRoot->AddGroup(LOCTEXT("WorkspaceMenu_DeveloperToolsCategory", "Instrumentation"), FSlateIcon()) )
+		, DeveloperToolsLogCategory ( MenuRoot->AddGroup(LOCTEXT("WorkspaceMenu_DeveloperToolsLogCategory", "Log"), FSlateIcon(), true))
 		, EditOptions( FWorkspaceItem::NewGroup(LOCTEXT( "WorkspaceEdit_Options", "Edit Options" )) )
 	{
+
+
 		ResetLevelEditorCategory();
 		ResetToolsCategory();
 	}
@@ -122,18 +146,23 @@ public:
 
 private:
 	TSharedPtr<FWorkspaceItem> MenuRoot;
+	TSharedPtr<FWorkspaceItem> ToolsMenuRoot;
 	
 	TSharedPtr<FWorkspaceItem> LevelEditorCategory;
 	TSharedPtr<FWorkspaceItem> LevelEditorViewportsCategory;
 	TSharedPtr<FWorkspaceItem> LevelEditorDetailsCategory;
 	TSharedPtr<FWorkspaceItem> LevelEditorCinematicsCategory;
+	TSharedPtr<FWorkspaceItem> LevelEditorVirtualProductionCategory;
 	TSharedPtr<FWorkspaceItem> LevelEditorModesCategory;
 
 	TSharedPtr<FWorkspaceItem> ToolsCategory;
 	TSharedPtr<FWorkspaceItem> DeveloperToolsCategory;
-	TSharedPtr<FWorkspaceItem> DeveloperToolsDebugCategory;
 	TSharedPtr<FWorkspaceItem> DeveloperToolsLogCategory;
+	TSharedPtr<FWorkspaceItem> DeveloperToolsSubMenuDebugCategory;
+	TSharedPtr<FWorkspaceItem> DeveloperToolsDebugCategory;
+	TSharedPtr<FWorkspaceItem> DeveloperToolsAuditCategory;
 	TSharedPtr<FWorkspaceItem> DeveloperToolsProfilingCategory;
+	TSharedPtr<FWorkspaceItem> DeveloperToolsPlatformsCategory;
 	TSharedPtr<FWorkspaceItem> DeveloperToolsMiscCategory;
 	
 	TSharedPtr<FWorkspaceItem> AutomationToolsCategory;
diff --git a/Engine/Source/Editor/WorkspaceMenuStructure/Public/WorkspaceMenuStructure.h b/Engine/Source/Editor/WorkspaceMenuStructure/Public/WorkspaceMenuStructure.h
--- a/Engine/Source/Editor/WorkspaceMenuStructure/Public/WorkspaceMenuStructure.h
+++ b/Engine/Source/Editor/WorkspaceMenuStructure/Public/WorkspaceMenuStructure.h
@@ -31,6 +31,9 @@ public:
 	/** Get the root of the menu structure. Pass this into PopulateTabSpawnerMenu() */
 	virtual TSharedRef<FWorkspaceItem> GetStructureRoot() const = 0;
 
+	/** Alternate root for the Tools menu */
+	virtual TSharedRef<FWorkspaceItem> GetToolsStructureRoot() const = 0;
+
 	/** See diagram above */
 	virtual TSharedRef<FWorkspaceItem> GetLevelEditorCategory() const = 0;
 
@@ -43,6 +46,9 @@ public:
 	/** See diagram above */
 	virtual TSharedRef<FWorkspaceItem> GetLevelEditorCinematicsCategory() const = 0;
 
+	/** See diagram above */
+	virtual TSharedRef<FWorkspaceItem> GetLevelEditorVirtualProductionCategory() const = 0;
+
 	/** See diagram above */
 	virtual TSharedRef<FWorkspaceItem> GetLevelEditorModesCategory() const = 0;
 
@@ -58,6 +64,12 @@ public:
 	/** See diagram above */
 	virtual TSharedRef<FWorkspaceItem> GetDeveloperToolsProfilingCategory() const = 0;	
 	
+	/** See diagram above */
+	virtual TSharedRef<FWorkspaceItem> GetDeveloperToolsAuditCategory() const = 0;	
+
+	/** See diagram above */
+	virtual TSharedRef<FWorkspaceItem> GetDeveloperToolsPlatformsCategory() const = 0;	
+	
 	/** See diagram above */
 	virtual TSharedRef<FWorkspaceItem> GetDeveloperToolsMiscCategory() const = 0;
 
