Nomad tabs docked with major tabs now have a menu instead of just empty space. The menu is simplified because nomad tabs do not belong to a specific editor.

diff --git a/Engine/Source/Editor/LevelEditor/LevelEditor.Build.cs b/Engine/Source/Editor/LevelEditor/LevelEditor.Build.cs
--- a/Engine/Source/Editor/LevelEditor/LevelEditor.Build.cs
+++ b/Engine/Source/Editor/LevelEditor/LevelEditor.Build.cs
@@ -81,6 +81,7 @@ public class LevelEditor : ModuleRules
 				"EnvironmentLightingViewer",
 				"DesktopPlatform",
 				"DataLayerEditor",
+				"TranslationEditor",
 			}
 		);
 
@@ -108,6 +109,7 @@ public class LevelEditor : ModuleRules
                 "Persona",
 				"LevelAssetEditor",
 				"MergeActors",
+
 			}
 		);
 
diff --git a/Engine/Source/Editor/LevelEditor/Private/LevelEditor.cpp b/Engine/Source/Editor/LevelEditor/Private/LevelEditor.cpp
--- a/Engine/Source/Editor/LevelEditor/Private/LevelEditor.cpp
+++ b/Engine/Source/Editor/LevelEditor/Private/LevelEditor.cpp
@@ -274,8 +274,9 @@ void FLevelEditorModule::StartupModule()
 
 	const IWorkspaceMenuStructure& MenuStructure = WorkspaceMenu::GetMenuStructure();
 
-	FGlobalTabmanager::Get()->RegisterTabSpawner("LevelEditor", FOnSpawnTab::CreateRaw( this, &FLevelEditorModule::SpawnLevelEditor ) )
-		.SetDisplayName( NSLOCTEXT("LevelEditor", "LevelEditorTab", "Level Editor") );
+	FGlobalTabmanager::Get()->RegisterTabSpawner("LevelEditor", FOnSpawnTab::CreateRaw(this, &FLevelEditorModule::SpawnLevelEditor))
+		.SetDisplayName(NSLOCTEXT("LevelEditor", "LevelEditorTab", "Level Editor"))
+		.SetAutoGenerateMenuEntry(false);
 
 	FModuleManager::LoadModuleChecked<ISlateReflectorModule>("SlateReflector").RegisterTabSpawner(MenuStructure.GetDeveloperToolsDebugCategory());
 
diff --git a/Engine/Source/Editor/LevelEditor/Private/LevelEditorMenu.cpp b/Engine/Source/Editor/LevelEditor/Private/LevelEditorMenu.cpp
--- a/Engine/Source/Editor/LevelEditor/Private/LevelEditorMenu.cpp
+++ b/Engine/Source/Editor/LevelEditor/Private/LevelEditorMenu.cpp
@@ -18,8 +18,10 @@
 #include "ToolMenus.h"
 #include "LevelEditorMenuContext.h"
 #include "Widgets/Input/SSpinBox.h"
-
 #include "Widgets/Input/SCheckBox.h"
+#include "Settings/EditorExperimentalSettings.h"
+#include "TranslationEditor/Public/ITranslationEditor.h"
+#include "LocalizationDashboard/Public/ILocalizationDashboardModule.h"
 
 #define LOCTEXT_NAMESPACE "LevelEditorMenu"
 
@@ -244,6 +246,37 @@ void FLevelEditorMenu::RegisterLevelEditorMenus()
 		static void ExtendToolsMenu()
 		{
 			UToolMenu* Menu = UToolMenus::Get()->RegisterMenu("LevelEditor.MainMenu.Tools", "MainFrame.MainMenu.Tools");
+
+			// Experimental section
+			{
+				// This is a temporary home for the spawners of experimental features that must be explicitly enabled.
+				// When the feature becomes permanent and need not check a flag, register a nomad spawner for it in the proper WorkspaceMenu category
+				const bool bTranslationPicker = GetDefault<UEditorExperimentalSettings>()->bEnableTranslationPicker;
+
+				// Make sure at least one is enabled before creating the section
+				if (bTranslationPicker)
+				{
+					FToolMenuSection& Section = Menu->AddSection("ExperimentalTabSpawners", LOCTEXT("ExperimentalTabSpawnersHeading", "Experimental"));
+					{
+						// Translation Picker
+						if (bTranslationPicker)
+						{
+							Section.AddMenuEntry(
+								"TranslationPicker",
+								LOCTEXT("TranslationPickerMenuItem", "Translation Picker"),
+								LOCTEXT("TranslationPickerMenuItemToolTip", "Launch the Translation Picker to Modify Editor Translations"),
+								FSlateIcon(),
+								FUIAction(FExecuteAction::CreateLambda(
+									[]()
+									{
+										FModuleManager::Get().LoadModuleChecked("TranslationEditor");
+										ITranslationEditor::OpenTranslationPicker();
+									}))
+							);
+						}
+					}
+				}
+			}
 		}
 	};
 
diff --git a/Engine/Source/Editor/LevelEditor/Private/SLevelEditor.cpp b/Engine/Source/Editor/LevelEditor/Private/SLevelEditor.cpp
--- a/Engine/Source/Editor/LevelEditor/Private/SLevelEditor.cpp
+++ b/Engine/Source/Editor/LevelEditor/Private/SLevelEditor.cpp
@@ -282,12 +282,8 @@ void SLevelEditor::Initialize( const TSharedRef<SDockTab>& OwnerTab, const TShar
 	LevelEditorModule.OnOverridePropertyEditorSelection().AddSP(this, &SLevelEditor::OnOverridePropertyEditorSelection);
 
 	TSharedRef<SWidget> ContentArea = RestoreContentArea( OwnerTab, OwnerWindow );
-	TSharedRef<SWidget> MenuBar =
-		SNew(SBox)
-		.AddMetaData<FTagMetaData>(FTagMetaData(TEXT("MainMenu")))
-		[
-			FLevelEditorMenu::MakeLevelEditorMenu(LevelEditorCommands, SharedThis(this))
-		];
+
+	FLevelEditorMenu::MakeLevelEditorMenu(LevelEditorCommands, SharedThis(this));
 		
 	ChildSlot
 	[
diff --git a/Engine/Source/Editor/LocalizationDashboard/LocalizationDashboard.Build.cs b/Engine/Source/Editor/LocalizationDashboard/LocalizationDashboard.Build.cs
--- a/Engine/Source/Editor/LocalizationDashboard/LocalizationDashboard.Build.cs
+++ b/Engine/Source/Editor/LocalizationDashboard/LocalizationDashboard.Build.cs
@@ -36,6 +36,7 @@ public class LocalizationDashboard : ModuleRules
 				"LocalizationService",
 				"InternationalizationSettings",
 				"ToolMenus",
+				"WorkspaceMenuStructure",
 			}
 		);
 
diff --git a/Engine/Source/Editor/LocalizationDashboard/Private/LocalizationDashboard.cpp b/Engine/Source/Editor/LocalizationDashboard/Private/LocalizationDashboard.cpp
--- a/Engine/Source/Editor/LocalizationDashboard/Private/LocalizationDashboard.cpp
+++ b/Engine/Source/Editor/LocalizationDashboard/Private/LocalizationDashboard.cpp
@@ -21,6 +21,8 @@
 #include "LocalizationDashboardSettings.h"
 #include "SSettingsEditorCheckoutNotice.h"
 #include "ToolMenus.h"
+#include "WorkspaceMenuStructureModule.h"
+#include "WorkspaceMenuStructure.h"
 
 #define LOCTEXT_NAMESPACE "LocalizationDashboard"
 
@@ -124,63 +126,64 @@ void SLocalizationDashboard::Construct(const FArguments& InArgs, const TSharedPt
 		return DockTab;
 	};
 	const TSharedRef<FWorkspaceItem> TargetSetsWorkspaceMenuCategory = TabManager->AddLocalWorkspaceMenuCategory(LOCTEXT("LocalizationDashboardWorkspaceMenuCategory", "Localization Dashboard"));
+
 	FTabSpawnerEntry& TargetsTabSpawnerEntry = TabManager->RegisterTabSpawner(TargetsDetailsTabName, FOnSpawnTab::CreateLambda(CreateTargetsTab))
 		.SetDisplayName(LOCTEXT("TargetsDetailTabSpawner", "Targets"));
+
 	TargetSetsWorkspaceMenuCategory->AddItem(TargetsTabSpawnerEntry.AsShared());
 
-	TSharedRef<FTabManager::FLayout> Layout = FTabManager::NewLayout("LocalizationDashboard_Experimental_V5")
+	TSharedRef<FTabManager::FLayout> Layout = FTabManager::NewLayout("LocalizationDashboard_v5")
 		->AddArea
 		(
-		FTabManager::NewPrimaryArea()
-		->SetOrientation(Orient_Horizontal)
-		->Split
-		(
-		FTabManager::NewStack()
-		->SetSizeCoefficient(1.0f)
-		->SetHideTabWell(true)
-		->AddTab(TargetsDetailsTabName, ETabState::OpenedTab)
-		)
-		->Split
-		(
-		FTabManager::NewStack()
-		->SetSizeCoefficient(2.0f)
-		->SetHideTabWell(false)
-		->AddTab(DocumentsTabName, ETabState::ClosedTab)
-		)
+			FTabManager::NewPrimaryArea()
+			->SetOrientation(Orient_Horizontal)
+			->Split
+			(
+				FTabManager::NewStack()
+				->SetSizeCoefficient(1.0f)
+				->SetHideTabWell(true)
+				->AddTab(TargetsDetailsTabName, ETabState::OpenedTab)
+			)
+			->Split
+			(
+				FTabManager::NewStack()
+				->SetSizeCoefficient(2.0f)
+				->SetHideTabWell(false)
+				->AddTab(DocumentsTabName, ETabState::ClosedTab)
+			)
 		);
 
 	Layout = FLayoutSaveRestore::LoadFromConfig(GEditorLayoutIni, Layout);
 
 	FToolMenuContext ToolMenuContext;
-	IMainFrameModule& MainFrameModule = FModuleManager::LoadModuleChecked<IMainFrameModule>( "MainFrame" );
-	const TSharedRef<SWidget> MenuWidget = MainFrameModule.MakeMainMenu( TabManager, "LocalizationDashboard.MainMenu", ToolMenuContext );
+	IMainFrameModule& MainFrameModule = FModuleManager::LoadModuleChecked<IMainFrameModule>("MainFrame");
+
+	TabManager->SetAllowWindowMenuBar(true);
+
+	UToolMenus::Get()->RegisterMenu("LocalizationDashboard.MainMenu", "MainFrame.NomadMainMenu");
 
 	FString ConfigFilePath;
 	ConfigFilePath = GetDefault<ULocalizationSettings>()->GetDefaultConfigFilename();
 	ConfigFilePath = FPaths::ConvertRelativePathToFull(ConfigFilePath);
 
 	ChildSlot
+	[
+		SNew(SVerticalBox)
+		+ SVerticalBox::Slot()
+		.AutoHeight()
+		.Padding(0.0f, 8.0f)
 		[
-			SNew(SVerticalBox)
-			+SVerticalBox::Slot()
-			.AutoHeight()
-			[
-				MenuWidget
-			]
-			+SVerticalBox::Slot()
-			.AutoHeight()
-			.Padding(0.0f, 8.0f)
-			[
-				SAssignNew(SettingsEditorCheckoutNotice, SSettingsEditorCheckoutNotice)
-				.ConfigFilePath(ConfigFilePath)
-			]
-			+SVerticalBox::Slot()
-			.FillHeight(1.0f)
-			[
-				TabManager->RestoreFrom(Layout, OwningWindow).ToSharedRef()
-			]
-		];
+			SAssignNew(SettingsEditorCheckoutNotice, SSettingsEditorCheckoutNotice)
+			.ConfigFilePath(ConfigFilePath)
+		]
+		+ SVerticalBox::Slot()
+		.FillHeight(1.0f)
+		[
+			TabManager->RestoreFrom(Layout, OwningWindow).ToSharedRef()
+		]
+	];
 
+	MainFrameModule.MakeMainMenu(TabManager, "LocalizationDashboard.MainMenu", ToolMenuContext);
 
 	// Open the first game target if present
 	if (ULocalizationTargetSet* GameTargetSet = ULocalizationSettings::GetGameTargetSet())
@@ -289,8 +292,11 @@ void FLocalizationDashboard::RegisterTabSpawner()
 		return DockTab;
 	};
 
-	FGlobalTabmanager::Get()->RegisterTabSpawner(TabName, FOnSpawnTab::CreateLambda( MoveTemp(SpawnMainTab) ) )
-		.SetDisplayName(LOCTEXT("MainTabTitle", "Localization Dashboard"));
+	FGlobalTabmanager::Get()->RegisterNomadTabSpawner(TabName, FOnSpawnTab::CreateLambda(MoveTemp(SpawnMainTab)))
+		.SetDisplayName(LOCTEXT("MainTabTitle", "Localization Dashboard"))
+		.SetTooltipText(LOCTEXT("LocalizationDashboardToolTip", "Open the Localization Dashboard for this Project."))
+		.SetIcon(FSlateIcon(FAppStyle::GetAppStyleSetName(), "LocalizationDashboard.MenuIcon"))
+		.SetGroup(WorkspaceMenu::GetMenuStructure().GetToolsCategory());
 }
 
 void FLocalizationDashboard::UnregisterTabSpawner()
diff --git a/Engine/Source/Editor/MainFrame/MainFrame.Build.cs b/Engine/Source/Editor/MainFrame/MainFrame.Build.cs
--- a/Engine/Source/Editor/MainFrame/MainFrame.Build.cs
+++ b/Engine/Source/Editor/MainFrame/MainFrame.Build.cs
@@ -58,7 +58,6 @@ public class MainFrame : ModuleRules
 				"LevelEditor",
 				"Settings",
 				"SourceCodeAccess",
-				"LocalizationDashboard",
 			}
 		);
 
@@ -78,7 +77,6 @@ public class MainFrame : ModuleRules
 				"LevelEditor",
 				"SourceCodeAccess",
 				"HotReload",
-				"LocalizationDashboard",
 			}
 		);
 	}
diff --git a/Engine/Source/Editor/MainFrame/Private/Frame/MainFrameActions.cpp b/Engine/Source/Editor/MainFrame/Private/Frame/MainFrameActions.cpp
--- a/Engine/Source/Editor/MainFrame/Private/Frame/MainFrameActions.cpp
+++ b/Engine/Source/Editor/MainFrame/Private/Frame/MainFrameActions.cpp
@@ -428,11 +428,6 @@ void FMainFrameActionCallbacks::ZipUpProject()
 	}
 }
 
-//void FMainFrameActionCallbacks::LocalizeProject()
-//{
-//	FModuleManager::LoadModuleChecked<ILocalizationDashboardModule>("LocalizationDashboard").Show();
-//}
-
 void FMainFrameActionCallbacks::SwitchProjectByIndex(int32 ProjectIndex)
 {
 	FUnrealEdMisc::Get().SwitchProject(RecentProjects[ProjectIndex].ProjectName);
diff --git a/Engine/Source/Editor/MainFrame/Private/MainFrameModule.cpp b/Engine/Source/Editor/MainFrame/Private/MainFrameModule.cpp
--- a/Engine/Source/Editor/MainFrame/Private/MainFrameModule.cpp
+++ b/Engine/Source/Editor/MainFrame/Private/MainFrameModule.cpp
@@ -258,6 +258,10 @@ void FMainFrameModule::CreateDefaultMainFrameAuxiliary(const bool bStartImmersiv
 				}
 			}
 
+
+			FToolMenuContext EmptyContext;
+			MakeMainMenu(FGlobalTabmanager::Get(), "MainFrame.NomadMainMenu", EmptyContext);
+		
 			MainFrameContent = FGlobalTabmanager::Get()->RestoreFrom(LoadedLayout, RootWindow, bEmbedTitleAreaContent, OutputCanBeNullptr);
 			// MainFrameContent will only be nullptr if its main area contains invalid tabs (probably some layout bug). If so, reset layout to avoid potential crashes
 			if (!MainFrameContent.IsValid())
@@ -330,15 +334,6 @@ TSharedRef<SWidget> FMainFrameModule::MakeMainMenu(const TSharedPtr<FTabManager>
 	return FMainMenu::MakeMainMenu(TabManager, MenuName, ToolMenuContext);
 }
 
-PRAGMA_DISABLE_DEPRECATION_WARNINGS
-TSharedRef<SWidget> FMainFrameModule::MakeMainTabMenu(const TSharedPtr<FTabManager>& TabManager, const FName MenuName, FToolMenuContext& ToolMenuContext) const
-{
-	return FMainMenu::MakeMainTabMenu(TabManager, MenuName, ToolMenuContext);
-}
-PRAGMA_ENABLE_DEPRECATION_WARNINGS
-
-
-BEGIN_SLATE_FUNCTION_BUILD_OPTIMIZATION
 TSharedRef<SWidget> FMainFrameModule::MakeDeveloperTools( const TArray<FMainFrameDeveloperTool>& AdditionalTools ) const
 {
 	struct Local
@@ -513,7 +508,6 @@ TSharedRef<SWidget> FMainFrameModule::MakeDeveloperTools( const TArray<FMainFram
 			]
 		];
 }
-END_SLATE_FUNCTION_BUILD_OPTIMIZATION
 
 
 void FMainFrameModule::SetLevelNameForWindowTitle( const FString& InLevelFileName )
@@ -809,6 +803,7 @@ void FMainFrameModule::HandleCodeAccessorLaunching()
 	CodeAccessorNotificationPtr.Pin()->SetCompletionState(SNotificationItem::CS_Pending);
 }
 
+
 void FMainFrameModule::HandleCodeAccessorOpenFileFailed(const FString& Filename)
 {
 	auto* Info = new FNotificationInfo(FText::Format(LOCTEXT("FileNotFound", "Could not find code file, {0}"), FText::FromString(Filename)));
diff --git a/Engine/Source/Editor/MainFrame/Private/MainFrameModule.h b/Engine/Source/Editor/MainFrame/Private/MainFrameModule.h
--- a/Engine/Source/Editor/MainFrame/Private/MainFrameModule.h
+++ b/Engine/Source/Editor/MainFrame/Private/MainFrameModule.h
@@ -15,7 +15,6 @@
 /**
  * Editor main frame module
  */
-PRAGMA_DISABLE_DEPRECATION_WARNINGS
 class FMainFrameModule
 	: public IMainFrameModule
 {
@@ -36,8 +35,6 @@ public:
 	virtual bool IsRecreatingDefaultMainFrame() const override;
 	virtual TSharedRef<SWidget> MakeMainMenu(const TSharedPtr<FTabManager>& TabManager, const FName MenuName, FToolMenuContext& ToolMenuContext) const override;
 	
-	// deprecated in 4.26
-	virtual TSharedRef<SWidget> MakeMainTabMenu(const TSharedPtr<FTabManager>& TabManager, const FName MenuName, FToolMenuContext& ToolMenuContext) const override;
 
 	virtual TSharedRef<SWidget> MakeDeveloperTools( const TArray<FMainFrameDeveloperTool>& AdditionalTools ) const override;
 
@@ -221,7 +218,6 @@ private:
 	{
 		DelayedShowMainFrameDelegate.Unbind();
 	}
-
 private:
 
 	// Weak pointer to the level editor's compile notification item.
@@ -266,5 +262,3 @@ private:
 	// Is recreating Default Main Frame
 	bool bRecreatingDefaultMainFrame;
 };
-PRAGMA_ENABLE_DEPRECATION_WARNINGS
-
diff --git a/Engine/Source/Editor/MainFrame/Private/Menus/MainFrameTranslationEditorMenu.h b/Engine/Source/Editor/MainFrame/Private/Menus/MainFrameTranslationEditorMenu.h
deleted file mode 100644
diff --git a/Engine/Source/Editor/MainFrame/Private/Menus/MainMenu.cpp b/Engine/Source/Editor/MainFrame/Private/Menus/MainMenu.cpp
--- a/Engine/Source/Editor/MainFrame/Private/Menus/MainMenu.cpp
+++ b/Engine/Source/Editor/MainFrame/Private/Menus/MainMenu.cpp
@@ -21,16 +21,14 @@
 #include "Menus/LayoutsMenu.h"
 #include "Menus/RecentProjectsMenu.h"
 #include "Menus/SettingsMenu.h"
-#include "Menus/MainFrameTranslationEditorMenu.h"
-
 #include "ToolMenus.h"
-
 #include "WorkspaceMenuStructure.h"
 #include "WorkspaceMenuStructureModule.h"
 #include "Features/EditorFeatures.h"
 #include "Features/IModularFeatures.h"
 #include "UndoHistoryModule.h"
 #include "Framework/Commands/GenericCommands.h"
+#include "ITranslationEditor.h"
 
 #define LOCTEXT_NAMESPACE "MainFileMenu"
 
@@ -176,59 +174,23 @@ void FMainMenu::RegisterWindowMenu()
 	{
 		if (USlateTabManagerContext* TabManagerContext = InData->FindContext<USlateTabManagerContext>())
 		{
-			TSharedPtr<FTabManager> TabManager = TabManagerContext->TabManager.Pin();
-			if (TabManager.IsValid())
+			
+			if (TSharedPtr<FTabManager> TabManager = TabManagerContext->TabManager.Pin())
 			{
-				// Local editor tabs
-				TabManager->PopulateLocalTabSpawnerMenu(InBuilder);
-
-				// General tabs
-				const IWorkspaceMenuStructure& MenuStructure = WorkspaceMenu::GetMenuStructure();
-				TabManager->PopulateTabSpawnerMenu(InBuilder, MenuStructure.GetStructureRoot());
-			}
-		}
-	}));
-
-	
-
-	// Experimental section
-	{
-		// This is a temporary home for the spawners of experimental features that must be explicitly enabled.
-		// When the feature becomes permanent and need not check a flag, register a nomad spawner for it in the proper WorkspaceMenu category
-		const bool bLocalizationDashboard = GetDefault<UEditorExperimentalSettings>()->bEnableLocalizationDashboard;
-		const bool bTranslationPicker = GetDefault<UEditorExperimentalSettings>()->bEnableTranslationPicker;
-
-		// Make sure at least one is enabled before creating the section
-		if (bLocalizationDashboard || bTranslationPicker)
-		{
-			FToolMenuSection& Section = Menu->AddSection("ExperimentalTabSpawners", LOCTEXT("ExperimentalTabSpawnersHeading", "Experimental"), FToolMenuInsert("WindowLayout", EToolMenuInsertType::Before));
-			{
-				// Localization Dashboard
-				if (bLocalizationDashboard)
+				// The global tab manager will be the tab manager for nomad tabs that appear docked as major tabs. However major tabs are not spawned
+				// via the window menu so ignore anything from the global tab manager since it is responsible for major tabs only
+				if(TabManager != FGlobalTabmanager::Get())
 				{
-					Section.AddMenuEntry(
-						"LocalizationDashboard",
-						LOCTEXT("LocalizationDashboardLabel", "Localization Dashboard"),
-						LOCTEXT("LocalizationDashboardToolTip", "Open the Localization Dashboard for this Project."),
-						FSlateIcon(FAppStyle::GetAppStyleSetName(), "LocalizationDashboard.MenuIcon"), 
-						FUIAction(FExecuteAction::CreateStatic(&FMainMenu::OpenLocalizationDashboard))
-						);
-				}
+					// Local editor tabs
+					TabManager->PopulateLocalTabSpawnerMenu(InBuilder);
 
-				// Translation Picker
-				if (bTranslationPicker)
-				{
-					Section.AddMenuEntry(
-						"TranslationPicker",
-						LOCTEXT("TranslationPickerMenuItem", "Translation Picker"),
-						LOCTEXT("TranslationPickerMenuItemToolTip", "Launch the Translation Picker to Modify Editor Translations"),
-						FSlateIcon(),
-						FUIAction(FExecuteAction::CreateStatic(&FMainFrameTranslationEditorMenu::HandleOpenTranslationPicker))
-						);
+					// General tabs
+					const IWorkspaceMenuStructure& MenuStructure = WorkspaceMenu::GetMenuStructure();
+					TabManager->PopulateTabSpawnerMenu(InBuilder, MenuStructure.GetStructureRoot());
 				}
 			}
 		}
-	}
+	}));
 
 	// Layout section
 	{
@@ -329,7 +291,9 @@ TSharedRef<SWidget> FMainMenu::MakeMainMenu(const TSharedPtr<FTabManager>& TabMa
 	// Cache all project names once
 	FMainFrameActionCallbacks::CacheProjectNames();
 
-	FMainMenu::RegisterMainMenu();
+	RegisterMainMenu();
+
+	RegisterNomadMainMenu();
 
 	ToolMenuContext.AppendCommandList(FMainFrameCommands::ActionList);
 
@@ -515,9 +479,6 @@ void FMainMenu::RegisterToolsMenu()
 			TSharedPtr<FTabManager> TabManager = TabManagerContext->TabManager.Pin();
 			if (TabManager.IsValid())
 			{
-				// Local editor tabs
-				TabManager->PopulateLocalTabSpawnerMenu(InBuilder);
-
 				// General tabs
 				const IWorkspaceMenuStructure& MenuStructure = WorkspaceMenu::GetMenuStructure();
 				TabManager->PopulateTabSpawnerMenu(InBuilder, MenuStructure.GetToolsStructureRoot());
@@ -525,6 +486,39 @@ void FMainMenu::RegisterToolsMenu()
 		}
 	}));
 
+
+	// Experimental section
+	{
+		// This is a temporary home for the spawners of experimental features that must be explicitly enabled.
+		// When the feature becomes permanent and need not check a flag, register a nomad spawner for it in the proper WorkspaceMenu category
+		//const bool bLocalizationDashboard = GetDefault<UEditorExperimentalSettings>()->bEnableLocalizationDashboard;
+		const bool bTranslationPicker = GetDefault<UEditorExperimentalSettings>()->bEnableTranslationPicker;
+
+		// Make sure at least one is enabled before creating the section
+		if ( bTranslationPicker)
+		{
+			FToolMenuSection& ExperimentalSection = Menu->AddSection("ExperimentalTabSpawners", LOCTEXT("ExperimentalTabSpawnersHeading", "Experimental"));
+			{
+				// Translation Picker
+				if (bTranslationPicker)
+				{
+					ExperimentalSection.AddMenuEntry(
+						"TranslationPicker",
+						LOCTEXT("TranslationPickerMenuItem", "Translation Picker"),
+						LOCTEXT("TranslationPickerMenuItemToolTip", "Launch the Translation Picker to Modify Editor Translations"),
+						FSlateIcon(),
+						FUIAction(FExecuteAction::CreateLambda(
+							[]()
+							{
+								FModuleManager::Get().LoadModuleChecked("TranslationEditor");
+								ITranslationEditor::OpenTranslationPicker();
+							}))
+					);
+				}
+			}
+		}
+	}
+
 }
 
 void FMainMenu::RegisterExitMenuItems()
@@ -544,11 +538,10 @@ void FMainMenu::RegisterExitMenuItems()
 #endif
 }
 
-PRAGMA_DISABLE_DEPRECATION_WARNINGS
-TSharedRef< SWidget > FMainMenu::MakeMainTabMenu( const TSharedPtr<FTabManager>& TabManager, const FName MenuName, FToolMenuContext& ToolMenuContext )
+void FMainMenu::RegisterNomadMainMenu()
 {
-	return FMainMenu::MakeMainMenu( TabManager, MenuName, ToolMenuContext );
+	UToolMenu* Menu = UToolMenus::Get()->RegisterMenu("MainFrame.NomadMainMenu", "MainFrame.MainMenu");
 }
-PRAGMA_ENABLE_DEPRECATION_WARNINGS
+
 
 #undef LOCTEXT_NAMESPACE
diff --git a/Engine/Source/Editor/MainFrame/Private/Menus/MainMenu.h b/Engine/Source/Editor/MainFrame/Private/Menus/MainMenu.h
--- a/Engine/Source/Editor/MainFrame/Private/Menus/MainMenu.h
+++ b/Engine/Source/Editor/MainFrame/Private/Menus/MainMenu.h
@@ -8,7 +8,6 @@
 #include "Modules/ModuleInterface.h"
 #include "Modules/ModuleManager.h"
 #include "Framework/Docking/TabManager.h"
-#include "ILocalizationDashboardModule.h"
 
 class FMenuBuilder;
 struct FToolMenuContext;
@@ -29,23 +28,12 @@ public:
 	 */
 	static TSharedRef<SWidget> MakeMainMenu( const TSharedPtr<FTabManager>& TabManager, const FName MenuName, FToolMenuContext& ToolMenuContext );
 
-	/**
-	 * Static: Creates a widget for the main tab's menu bar.  This is just like the main menu bar, but also includes.
-	 * some "project level" menu items that we don't want propagated to most normal menus.
-	 *
-	 * @param TabManager Create the workspace menu based on this tab manager.
-	 * @param MenuName Identifier associated with the menu.
-	 * @return ToolMenuContext Context containing state.
-	 */
-	UE_DEPRECATED(4.26, "MakeMainTabMenu has been deprecated, use MakeMainMenu instead.")
-	static TSharedRef<SWidget> MakeMainTabMenu( const TSharedPtr<FTabManager>& TabManager, const FName MenuName, FToolMenuContext& ToolMenuContext );
-
 	/**
 	 * Static: Registers main menu with menu system.
 	 */
 	static void RegisterMainMenu();
 
-protected:
+private:
 
 	/**
 	 * Called to fill the file menu's content.
@@ -83,7 +71,11 @@ protected:
 	 */
 	static void RegisterExitMenuItems();
 
-private:
+	/**
+	 * Called to fill the main menu for nomad tabs (these tabs generally have a less menu options)
+	 */
+	static void RegisterNomadMainMenu();
+
 	/** 
 	* Opens the experimental project launcher tab.
 	* Remove this when it is is no longer experimental.
@@ -93,15 +85,6 @@ private:
 		FGlobalTabmanager::Get()->TryInvokeTab(FName(TEXT("ProjectLauncher")));
 	}
 
-	/**
-	* Opens the experimental localization dashboard.
-	* Remove this when it is no longer experimental.
-	*/
-	static void OpenLocalizationDashboard()
-	{
-		FModuleManager::LoadModuleChecked<ILocalizationDashboardModule>("LocalizationDashboard").Show();
-	}
-
 	/**
 	* Opens the experimental Visual Logger tab.
 	* Remove this when it is no longer experimental.
diff --git a/Engine/Source/Editor/MainFrame/Public/Interfaces/IMainFrameModule.h b/Engine/Source/Editor/MainFrame/Public/Interfaces/IMainFrameModule.h
--- a/Engine/Source/Editor/MainFrame/Public/Interfaces/IMainFrameModule.h
+++ b/Engine/Source/Editor/MainFrame/Public/Interfaces/IMainFrameModule.h
@@ -78,8 +78,6 @@ public:
 	 *
 	 * @return	The newly-created menu widget
 	 */
-	UE_DEPRECATED(4.26, "MakeMainTabMenu has been deprecated.  Use MakeMainMenu instead")
-	virtual TSharedRef<SWidget> MakeMainTabMenu( const TSharedPtr<FTabManager>& TabManager, const FName MenuName, FToolMenuContext& ToolMenuContext ) const = 0;
 
 	/**
 	 * Generates a menu for status and developer
diff --git a/Engine/Source/Editor/StatusBar/StatusBar.Build.cs b/Engine/Source/Editor/StatusBar/StatusBar.Build.cs
--- a/Engine/Source/Editor/StatusBar/StatusBar.Build.cs
+++ b/Engine/Source/Editor/StatusBar/StatusBar.Build.cs
@@ -4,27 +4,37 @@ using UnrealBuildTool;
 
 public class StatusBar : ModuleRules
 {
-    public StatusBar(ReadOnlyTargetRules Target)
-         : base(Target)
-    {
+	public StatusBar(ReadOnlyTargetRules Target)
+		 : base(Target)
+	{
 		PublicIncludePaths.Add(ModuleDirectory + "/Public");
 
 		PrivateDependencyModuleNames.AddRange(
-            new string[] {
-                "Core",
-                "CoreUObject",
-                "Engine",           
+			new string[] {
+				"Core",
+				"CoreUObject",
+				"Engine",
 				"Slate",
-                "SlateCore",
+				"SlateCore",
 				"InputCore",
-                "EditorStyle",
+				"EditorStyle",
 				"EditorFramework",
-                "UnrealEd",
+				"UnrealEd",
 				"ToolMenus",
 				"OutputLog",
 				"SourceControlWindows",
 				"EditorSubsystem",
+			});
+
+		DynamicallyLoadedModuleNames.AddRange(
+			new string[] {
+				"MainFrame",
+			});
+
+
+		PrivateIncludePathModuleNames.AddRange(
+			new string[] {
 				"MainFrame",
-            });
-    }
+			});
+	}
 }
diff --git a/Engine/Source/Editor/UnrealEd/Classes/Settings/EditorExperimentalSettings.h b/Engine/Source/Editor/UnrealEd/Classes/Settings/EditorExperimentalSettings.h
--- a/Engine/Source/Editor/UnrealEd/Classes/Settings/EditorExperimentalSettings.h
+++ b/Engine/Source/Editor/UnrealEd/Classes/Settings/EditorExperimentalSettings.h
@@ -47,10 +47,6 @@ public:
 	/** Allows usage of the procedural foliage system */
 	UPROPERTY(EditAnywhere, config, Category = Foliage, meta = (DisplayName = "Procedural Foliage"))
 	bool bProceduralFoliage;
-		
-	/** Allows usage of the Localization Dashboard */
-	UPROPERTY(EditAnywhere, config, Category = Tools, meta = (DisplayName = "Localization Dashboard"))
-	bool bEnableLocalizationDashboard;
 
 	/** Allows usage of the Translation Picker */
 	UPROPERTY(EditAnywhere, config, Category = Tools, meta = (DisplayName = "Translation Picker"))
diff --git a/Engine/Source/Editor/UnrealEd/Private/Settings/SettingsClasses.cpp b/Engine/Source/Editor/UnrealEd/Private/Settings/SettingsClasses.cpp
--- a/Engine/Source/Editor/UnrealEd/Private/Settings/SettingsClasses.cpp
+++ b/Engine/Source/Editor/UnrealEd/Private/Settings/SettingsClasses.cpp
@@ -155,7 +155,6 @@ UEditorExperimentalSettings::UEditorExperimentalSettings( const FObjectInitializ
 	, bEnableAsyncSkeletalMeshCompilation(false)
 	, bHDREditor(false)
 	, HDREditorNITLevel(160.0f)
-	, bEnableLocalizationDashboard(true)
 	, bUseOpenCLForConvexHullDecomp(false)
 	, bAllowPotentiallyUnsafePropertyEditing(false)
 {
diff --git a/Engine/Source/Runtime/Slate/Private/Framework/Docking/SDockingArea.cpp b/Engine/Source/Runtime/Slate/Private/Framework/Docking/SDockingArea.cpp
--- a/Engine/Source/Runtime/Slate/Private/Framework/Docking/SDockingArea.cpp
+++ b/Engine/Source/Runtime/Slate/Private/Framework/Docking/SDockingArea.cpp
@@ -650,7 +650,7 @@ void SDockingArea::UpdateWindowChromeAndSidebar()
 
 			if (ParentWindow->GetTitleBar())
 			{
-				ParentWindow->GetTitleBar()->SetUseLargeIcon(bAccountForMenuBarPadding);
+				ParentWindow->GetTitleBar()->SetAllowMenuBar(bAccountForMenuBarPadding);
 			}
 		}
 	}
diff --git a/Engine/Source/Runtime/Slate/Private/Framework/Docking/TabManager.cpp b/Engine/Source/Runtime/Slate/Private/Framework/Docking/TabManager.cpp
--- a/Engine/Source/Runtime/Slate/Private/Framework/Docking/TabManager.cpp
+++ b/Engine/Source/Runtime/Slate/Private/Framework/Docking/TabManager.cpp
@@ -1000,7 +1000,9 @@ TSharedPtr<SWidget> FTabManager::RestoreFrom(const TSharedRef<FLayout>& Layout,
 	}
 
 	UpdateStats();
-	
+
+	FinishRestore();
+
 	return PrimaryDockArea;
 }
 
@@ -2540,7 +2542,9 @@ TSharedRef<FTabManager> FGlobalTabmanager::NewTabManager( const TSharedRef<SDock
 
 	const TSharedRef<FTabManager> NewTabManager = FTabManager::New( InOwnerTab, NomadTabSpawner );
 	SubTabManagers.Add( FSubTabManager(InOwnerTab, NewTabManager) );
+
 	UpdateStats();
+
 	return NewTabManager;
 }
 
@@ -2745,6 +2749,17 @@ void FGlobalTabmanager::OpenUnmanagedTab(FName PlaceholderId, const FSearchPrefe
 	}
 }
 
+void FGlobalTabmanager::FinishRestore()
+{
+	for (FSubTabManager& SubManagerInfo : SubTabManagers)
+	{
+		if (TSharedPtr<FTabManager> Manager = SubManagerInfo.TabManager.Pin())
+		{
+			Manager->UpdateMainMenu(nullptr, false);
+		}
+	}
+}
+
 void FGlobalTabmanager::SetProxyTabManager(TSharedPtr<FProxyTabmanager> InProxyTabManager)
 {
 	ProxyTabManager = InProxyTabManager;
diff --git a/Engine/Source/Runtime/Slate/Public/Framework/Application/SWindowTitleBar.h b/Engine/Source/Runtime/Slate/Public/Framework/Application/SWindowTitleBar.h
--- a/Engine/Source/Runtime/Slate/Public/Framework/Application/SWindowTitleBar.h
+++ b/Engine/Source/Runtime/Slate/Public/Framework/Application/SWindowTitleBar.h
@@ -160,8 +160,7 @@ public:
 	}
 
 	//~ Begin IWindowTitleBar Interface
-
-	void Flash( ) override
+	void Flash() override
 	{
 		TitleFlashSequence = FCurveSequence(0, SWindowTitleBarDefs::WindowFlashDuration, ECurveEaseFunction::Linear);
 		TitleFlashSequence.Play(this->AsShared());
@@ -169,7 +168,7 @@ public:
 
 	virtual void UpdateWindowMenu(TSharedPtr<SWidget> MenuContent) override
 	{
-		if(MenuContent.IsValid())
+		if(MenuContent.IsValid() && bAllowMenuBar)
 		{
 			(*WindowMenuSlot)
 			[
@@ -209,11 +208,14 @@ public:
 		}
 	}
 
-	virtual void SetUseLargeIcon(bool bUseLargeIcon) override
+
+	virtual void SetAllowMenuBar(bool bInAllowMenuBar)
 	{
-		if(AppIconWidget)
+		bAllowMenuBar = bInAllowMenuBar;
+
+		if (AppIconWidget)
 		{
-			if (bUseLargeIcon)
+			if (bInAllowMenuBar)
 			{
 				AppIconWidget->SetIcon(FSlateApplicationBase::Get().GetAppIcon(), FAppStyle::Get().GetMargin("AppIconPadding", nullptr, FMargin(0)));
 			}
@@ -339,7 +341,7 @@ protected:
 				.Expose(WindowMenuSlot);
 
 			// Default everything to use thge small icon unless specifically set to use the large icon
-			SetUseLargeIcon(false);
+			SetAllowMenuBar(false);
 		}
 		else
 		{
@@ -760,5 +762,5 @@ private:
 
 	TSharedPtr<SAppIconWidget> AppIconWidget;
 	bool ShowAppIcon;
-
+	bool bAllowMenuBar = false;
 }; 
diff --git a/Engine/Source/Runtime/Slate/Public/Framework/Docking/TabManager.h b/Engine/Source/Runtime/Slate/Public/Framework/Docking/TabManager.h
--- a/Engine/Source/Runtime/Slate/Public/Framework/Docking/TabManager.h
+++ b/Engine/Source/Runtime/Slate/Public/Framework/Docking/TabManager.h
@@ -1020,6 +1020,9 @@ class SLATE_API FTabManager : public TSharedFromThis<FTabManager>
 		/** Called when tab(s) have been added or windows created */
 		virtual void UpdateStats();
 		
+		/** Called at the end of RestoreFrom for tab managers to complete any work after all tabs have been restored */
+		virtual void FinishRestore() {};
+
 	private:
 		/** Checks all dock areas and adds up the number of open tabs and unique parent windows in the manager */
 		void GetRecordableStats( int32& OutTabCount, TArray<TSharedPtr<SWindow>>& OutUniqueParentWindows ) const;
@@ -1220,6 +1223,8 @@ protected:
 
 	virtual void OpenUnmanagedTab(FName PlaceholderId, const FSearchPreference& SearchPreference, const TSharedRef<SDockTab>& UnmanagedTab) override;
 
+	virtual void FinishRestore() override;
+
 public:
 	virtual void OnTabManagerClosing() override;
 
diff --git a/Engine/Source/Runtime/SlateCore/Public/Application/SlateApplicationBase.h b/Engine/Source/Runtime/SlateCore/Public/Application/SlateApplicationBase.h
--- a/Engine/Source/Runtime/SlateCore/Public/Application/SlateApplicationBase.h
+++ b/Engine/Source/Runtime/SlateCore/Public/Application/SlateApplicationBase.h
@@ -54,7 +54,7 @@ public:
 	virtual void Flash( ) = 0;
 	virtual void UpdateWindowMenu(TSharedPtr<SWidget> MenuContent) = 0;
 	virtual void UpdateBackgroundContent(TSharedPtr<SWidget> BackgroundContent) = 0;
-	virtual void SetUseLargeIcon(bool bUseSmallIcon) = 0;
+	virtual void SetAllowMenuBar(bool bInAllowMenuBar) = 0;
 };
 
 class FSlateApplicationBase;
