Resolve circular depends issue for Linux/Mac

diff --git a/Engine/Source/Editor/MainFrame/MainFrame.Build.cs b/Engine/Source/Editor/MainFrame/MainFrame.Build.cs
--- a/Engine/Source/Editor/MainFrame/MainFrame.Build.cs
+++ b/Engine/Source/Editor/MainFrame/MainFrame.Build.cs
@@ -58,6 +58,7 @@ public class MainFrame : ModuleRules
 				"LevelEditor",
 				"Settings",
 				"SourceCodeAccess",
+				"LocalizationDashboard", // not required but causes circular depends issues on Linux/Mac
 			}
 		);
 
@@ -77,6 +78,7 @@ public class MainFrame : ModuleRules
 				"LevelEditor",
 				"SourceCodeAccess",
 				"HotReload",
+				"LocalizationDashboard", // not required but causes circular depends issues on Linux/Mac
 			}
 		);
 	}
