Missing file from CL 16452199

diff --git a/Engine/Source/Editor/MainFrame/Private/Menus/LayoutsMenu.cpp b/Engine/Source/Editor/MainFrame/Private/Menus/LayoutsMenu.cpp
--- a/Engine/Source/Editor/MainFrame/Private/Menus/LayoutsMenu.cpp
+++ b/Engine/Source/Editor/MainFrame/Private/Menus/LayoutsMenu.cpp
@@ -784,7 +784,7 @@ void FLayoutsMenuLoad::ReloadCurrentLayout()
 	if (bAreAssetEditorOpened)
 	{
 		// Save open asset editors + disable manual saving (no need to close them with AssetEditorSubsystem->CloseAllAssetEditors)
-		AssetEditorSubsystem->SaveOpenAssetEditors(/*bOnShutdown*/true, /*bCancelIfDebugger*/false);
+		AssetEditorSubsystem->SaveOpenAssetEditors(true);
 		AssetEditorSubsystem->SetAutoRestoreAndDisableSaving(true);
 	}
 	FUnrealEdMisc::Get().AllowSavingLayoutOnClose(false);
