Updates Edit menu to dynamically change between "Connect to Source Control / Change Source Control Settings" depending on the current state of the connexion

diff --git a/Engine/Source/Editor/MainFrame/Private/Frame/MainFrameActions.cpp b/Engine/Source/Editor/MainFrame/Private/Frame/MainFrameActions.cpp
--- a/Engine/Source/Editor/MainFrame/Private/Frame/MainFrameActions.cpp
+++ b/Engine/Source/Editor/MainFrame/Private/Frame/MainFrameActions.cpp
@@ -102,9 +102,12 @@ void FMainFrameCommands::RegisterCommands()
 	UI_COMMAND( ViewChangelists, "View Changelists...", "Opens a dialog with check in options for content and levels", EUserInterfaceActionType::Button, FInputChord() );
 	ActionList->MapAction(ViewChangelists, FExecuteAction::CreateStatic( &FMainFrameActionCallbacks::ViewChangelists ), FCanExecuteAction::CreateStatic( &FMainFrameActionCallbacks::CanViewChangelists ) );
 
-	UI_COMMAND( ConnectToSourceControl, "Connect To Source Control...", "Connect to source control to allow source control operations to be performed on content and levels.", EUserInterfaceActionType::Button, FInputChord() );
+	UI_COMMAND( ConnectToSourceControl, "Connect to Source Control...", "Connect to source control to allow source control operations to be performed on content and levels.", EUserInterfaceActionType::Button, FInputChord() );
 	ActionList->MapAction( ConnectToSourceControl, FExecuteAction::CreateStatic( &FMainFrameActionCallbacks::ConnectToSourceControl ), DefaultExecuteAction );
 
+	UI_COMMAND( ChangeSourceControlSettings, "Change Source Control Settings...", "Opens a dialog to change source control settings.", EUserInterfaceActionType::Button, FInputChord() );
+	ActionList->MapAction( ChangeSourceControlSettings, FExecuteAction::CreateStatic(&FMainFrameActionCallbacks::ConnectToSourceControl), DefaultExecuteAction);
+
 	UI_COMMAND( NewProject, "New Project...", "Opens a dialog to create a new game project", EUserInterfaceActionType::Button, FInputChord() );
 	ActionList->MapAction( NewProject, FExecuteAction::CreateStatic( &FMainFrameActionCallbacks::NewProject, false, true), DefaultExecuteAction );
 
diff --git a/Engine/Source/Editor/MainFrame/Private/Frame/MainFrameActions.h b/Engine/Source/Editor/MainFrame/Private/Frame/MainFrameActions.h
--- a/Engine/Source/Editor/MainFrame/Private/Frame/MainFrameActions.h
+++ b/Engine/Source/Editor/MainFrame/Private/Frame/MainFrameActions.h
@@ -28,6 +28,7 @@ public:
 	TSharedPtr< FUICommandInfo > ChooseFilesToSave;
 	TSharedPtr< FUICommandInfo > ViewChangelists;
 	TSharedPtr< FUICommandInfo > ConnectToSourceControl;
+	TSharedPtr< FUICommandInfo > ChangeSourceControlSettings;
 	TSharedPtr< FUICommandInfo > NewProject;
 	TSharedPtr< FUICommandInfo > OpenProject;
 	TSharedPtr< FUICommandInfo > AddCodeToProject;
diff --git a/Engine/Source/Editor/MainFrame/Private/Menus/MainMenu.cpp b/Engine/Source/Editor/MainFrame/Private/Menus/MainMenu.cpp
--- a/Engine/Source/Editor/MainFrame/Private/Menus/MainMenu.cpp
+++ b/Engine/Source/Editor/MainFrame/Private/Menus/MainMenu.cpp
@@ -156,7 +156,29 @@ void FMainMenu::RegisterEditMenu()
 			}
 		}));
 
-		Section.AddMenuEntry(FMainFrameCommands::Get().ConnectToSourceControl);
+		Section.AddDynamicEntry("ConnectToSourceControl", FNewToolMenuSectionDelegate::CreateLambda([](FToolMenuSection& InSection)
+		{
+			ISourceControlModule& SourceControlModule = ISourceControlModule::Get();
+			
+			if (ISourceControlModule::Get().IsEnabled() && ISourceControlModule::Get().GetProvider().IsAvailable())
+			{
+				InSection.AddMenuEntry(
+					FMainFrameCommands::Get().ChangeSourceControlSettings,
+					TAttribute<FText>(),
+					TAttribute<FText>(),
+					FSlateIcon(FEditorStyle::GetStyleSetName(), "SourceControl.Actions.ChangeSettings")
+				);
+			}
+			else
+			{
+				InSection.AddMenuEntry(
+					FMainFrameCommands::Get().ConnectToSourceControl,
+					TAttribute<FText>(),
+					TAttribute<FText>(),
+					FSlateIcon(FEditorStyle::GetStyleSetName(), "SourceControl.Actions.Connect")
+				);
+			}
+		}));
 	}
 }
 
