Reorganizes Help menus and cleans up external links in BaseEditor.ini.
Disables the tutorial system in order to remove its entry points from menus.

diff --git a/Engine/Config/BaseEditor.ini b/Engine/Config/BaseEditor.ini
--- a/Engine/Config/BaseEditor.ini
+++ b/Engine/Config/BaseEditor.ini
@@ -274,22 +274,17 @@ RetainJpegFormat=True
 bShouldPromptOnCurveRemoveAll=true
 
 [UnrealEd.URLs]
-UDNURL="http://udn.unrealengine.com"
-UDNDocsURL="https://docs.unrealengine.com/{VERSION}/{I18N}/"
+DocumentationURL="https://docs.unrealengine.com/{VERSION}/{I18N}/"
 OnlineLearningURL="https://www.unrealengine.com/onlinelearning-courses"
 ForumsURL="http://forums.unrealengine.com"
 APIDocsURL="https://docs.unrealengine.com/{VERSION}/{I18N}/API/"
-MapErrorURL="https://docs.unrealengine.com/{VERSION}/{I18N}/"
-DocumentationURL="http://docs.unrealengine.com/"
-AskAQuestionURL="http://udn.unrealengine.com/questions/ask.html"
-SearchForAnswersURL="http://udn.unrealengine.com/search.html"
-SupportWebsiteURL="http://udn.unrealengine.com/index.html"
+SearchForAnswersURL="http://answers.unrealengine.com"
+SupportWebsiteURL="http://unrealengine.com/support"
 SourceCodeIDEURL_Windows="https://aka.ms/vs/16/release/vs_Community.exe"
 SourceCodeIDEURL_Mac="https://developer.apple.com/xcode"
 SourceCodeIDEURL_Linux="https://code.visualstudio.com/docs/setup/linux"
 SourceCodeIDEURL_Other="http://en.wikipedia.org/wiki/Integrated_development_environment"
 EpicGamesURL="http://UnrealEngine.com"
-FacebookURL="http://www.facebook.com/UnrealEngine"
 ViewportControlsURL="https://docs.unrealengine.com/{VERSION}/{I18N}/BuildingWorlds/LevelEditor/Viewports/ViewportControls/"
 MeshSimplificationPluginsURL="/ue/marketplace/content-cat/assets/codeplugins"
 ReportABugURL="https://epicsupport.force.com/unrealengine/s/"
@@ -297,11 +292,7 @@ IssueTrackerURL="https://issues.unrealengine.com/"
 
 
 [UnrealEd.URLOverrides]
-UDNURL="http://docs.unrealengine.com/"
-DocumentationURL="http://docs.unrealengine.com/"
-AskAQuestionURL="http://answers.unrealengine.com"
-SearchForAnswersURL="http://answers.unrealengine.com"
-SupportWebsiteURL="http://unrealengine.com/support"
+; Allows you to override the locations in [UnrealEd.URLs] for binaries installed through the Launcher.
 
 [UserDefinedStructure]
 bUseUserDefinedStructure=true
diff --git a/Engine/Source/Editor/Documentation/Private/DocumentationLink.h b/Engine/Source/Editor/Documentation/Private/DocumentationLink.h
--- a/Engine/Source/Editor/Documentation/Private/DocumentationLink.h
+++ b/Engine/Source/Editor/Documentation/Private/DocumentationLink.h
@@ -15,7 +15,7 @@ public:
 	static FString GetUrlRoot()
 	{
 		FString Url;
-		FUnrealEdMisc::Get().GetURL( TEXT("UDNDocsURL"), Url, true );
+		FUnrealEdMisc::Get().GetURL( TEXT("DocumentationURL"), Url, true );
 
 		if ( !Url.EndsWith( TEXT( "/" ) ) )
 		{
@@ -33,7 +33,7 @@ public:
 	static FString GetHomeUrl(const FCultureRef& Culture)
 	{
 		FString Url;
-		FUnrealEdMisc::Get().GetURL( TEXT("UDNURL"), Url, true );
+		FUnrealEdMisc::Get().GetURL( TEXT("DocumentationURL"), Url, true );
 		FUnrealEdMisc::Get().ReplaceDocumentationURLWildcards(Url, Culture);
 
 		return Url;
diff --git a/Engine/Source/Editor/EditorStyle/Private/StarshipStyle.cpp b/Engine/Source/Editor/EditorStyle/Private/StarshipStyle.cpp
--- a/Engine/Source/Editor/EditorStyle/Private/StarshipStyle.cpp
+++ b/Engine/Source/Editor/EditorStyle/Private/StarshipStyle.cpp
@@ -4285,30 +4285,30 @@ void FStarshipEditorStyle::FStyle::SetupLevelEditorStyle()
 	// Level editor tool bar icons
 #if WITH_EDITOR
 	{
-		Set("LevelEditor.BrowseDocumentation",     new IMAGE_BRUSH_SVG("Starship/Common/Documentation", Icon16x16));
-		Set("LevelEditor.BrowseAPIReference",      new IMAGE_BRUSH_SVG("Starship/Common/Documentation", Icon16x16));
-		Set("LevelEditor.BrowseCVars",             new IMAGE_BRUSH_SVG("Starship/Common/Console", Icon16x16));
-		Set("LevelEditor.Tutorials",               new IMAGE_BRUSH_SVG("Starship/Common/Tutorials", Icon16x16));
-		Set("LevelEditor.BrowseViewportControls",  new IMAGE_BRUSH_SVG("Starship/Common/ViewportControls", Icon16x16));
+		Set("LevelEditor.BrowseDocumentation",      new IMAGE_BRUSH_SVG("Starship/Common/Documentation", Icon16x16));
+		Set("LevelEditor.Tutorials",                new IMAGE_BRUSH_SVG("Starship/Common/Tutorials", Icon16x16));
+		Set("LevelEditor.BrowseViewportControls",   new IMAGE_BRUSH_SVG("Starship/Common/ViewportControls", Icon16x16));
 		Set("LevelEditor.PasteHere",				new IMAGE_BRUSH_SVG("Starship/Actors/paste-here", Icon16x16));
 
-		Set("MainFrame.ToggleFullscreen",          new IMAGE_BRUSH_SVG("Starship/Common/EnableFullscreen", Icon16x16));
-		Set("MainFrame.LoadLayout",                new IMAGE_BRUSH_SVG("Starship/Common/LayoutLoad", Icon16x16));
-		Set("MainFrame.SaveLayout",                new IMAGE_BRUSH_SVG("Starship/Common/LayoutSave", Icon16x16));
-		Set("MainFrame.RemoveLayout",              new IMAGE_BRUSH_SVG("Starship/Common/LayoutRemove", Icon16x16));
-
-		Set("MainFrame.OpenIssueTracker",          new IMAGE_BRUSH_SVG("Starship/Common/IssueTracker", Icon16x16));
-		Set("MainFrame.ReportABug",                new IMAGE_BRUSH_SVG("Starship/Common/Bug", Icon16x16));
-
-		Set("MainFrame.VisitAskAQuestionPage",     new IMAGE_BRUSH_SVG("Starship/Menus/Bug", Icon16x16));
-		Set("MainFrame.VisitOnlineLearning",	   new IMAGE_BRUSH_SVG("Starship/Common/Documentation", Icon16x16));
-		Set("MainFrame.VisitWiki",                 new IMAGE_BRUSH_SVG("Starship/Common/WebBrowser", Icon16x16));
-		Set("MainFrame.VisitForums",               new IMAGE_BRUSH_SVG("Starship/Common/WebBrowser", Icon16x16));
-		Set("MainFrame.VisitSearchForAnswersPage", new IMAGE_BRUSH_SVG("Starship/Common/WebBrowser", Icon16x16));
-		Set("MainFrame.VisitSupportWebSite",       new IMAGE_BRUSH_SVG("Starship/Common/Support", Icon16x16));
-		Set("MainFrame.VisitEpicGamesDotCom",      new IMAGE_BRUSH_SVG("About/EpicGamesLogo", Icon16x16));
-		Set("MainFrame.AboutUnrealEd",             new IMAGE_BRUSH_SVG("About/UnrealLogo", Icon16x16));
-		Set("MainFrame.CreditsUnrealEd",           new IMAGE_BRUSH_SVG("Starship/Common/Credits", Icon16x16));
+		Set("MainFrame.ToggleFullscreen",           new IMAGE_BRUSH_SVG("Starship/Common/EnableFullscreen", Icon16x16));
+		Set("MainFrame.LoadLayout",                 new IMAGE_BRUSH_SVG("Starship/Common/LayoutLoad", Icon16x16));
+		Set("MainFrame.SaveLayout",                 new IMAGE_BRUSH_SVG("Starship/Common/LayoutSave", Icon16x16));
+		Set("MainFrame.RemoveLayout",               new IMAGE_BRUSH_SVG("Starship/Common/LayoutRemove", Icon16x16));
+
+		Set("MainFrame.OpenIssueTracker",           new IMAGE_BRUSH_SVG("Starship/Common/IssueTracker", Icon16x16));
+		Set("MainFrame.ReportABug",                 new IMAGE_BRUSH_SVG("Starship/Common/Bug", Icon16x16));
+
+		Set("SystemWideCommands.OpenDocumentation", new IMAGE_BRUSH_SVG("Starship/Common/Documentation", Icon16x16));
+		Set("MainFrame.DocumentationHome",	        new IMAGE_BRUSH_SVG("Starship/Common/Documentation", Icon16x16));
+		Set("MainFrame.VisitOnlineLearning",	    new IMAGE_BRUSH_SVG("Starship/Common/Documentation", Icon16x16));
+		Set("MainFrame.BrowseAPIReference",         new IMAGE_BRUSH_SVG("Starship/Common/Documentation", Icon16x16));
+		Set("MainFrame.BrowseCVars",                new IMAGE_BRUSH_SVG("Starship/Common/Console", Icon16x16));
+		Set("MainFrame.VisitForums",                new IMAGE_BRUSH_SVG("Starship/Common/WebBrowser", Icon16x16));
+		Set("MainFrame.VisitSearchForAnswersPage",  new IMAGE_BRUSH_SVG("Starship/Common/WebBrowser", Icon16x16));
+		Set("MainFrame.VisitSupportWebSite",        new IMAGE_BRUSH_SVG("Starship/Common/Support", Icon16x16));
+		Set("MainFrame.VisitEpicGamesDotCom",       new IMAGE_BRUSH_SVG("About/EpicGamesLogo", Icon16x16));
+		Set("MainFrame.AboutUnrealEd",              new IMAGE_BRUSH_SVG("About/UnrealLogo", Icon16x16));
+		Set("MainFrame.CreditsUnrealEd",            new IMAGE_BRUSH_SVG("Starship/Common/Credits", Icon16x16));
 
 		Set( "EditorViewport.SelectMode", new IMAGE_BRUSH_SVG("Starship/EditorViewport/select", Icon16x16) );
 		Set( "EditorViewport.TranslateMode", new IMAGE_BRUSH_SVG( "Starship/EditorViewport/translate", Icon16x16 ) );
diff --git a/Engine/Source/Editor/IntroTutorials/Private/IntroTutorials.cpp b/Engine/Source/Editor/IntroTutorials/Private/IntroTutorials.cpp
--- a/Engine/Source/Editor/IntroTutorials/Private/IntroTutorials.cpp
+++ b/Engine/Source/Editor/IntroTutorials/Private/IntroTutorials.cpp
@@ -48,7 +48,7 @@ FIntroTutorials::FIntroTutorials()
 	: CurrentObjectClass(nullptr)
 	, ContentIntroCurve(nullptr)
 {
-	bDisableTutorials = false;
+	bDisableTutorials = true;
 
 	bDesireResettingTutorialSeenFlagOnLoad = FParse::Param(FCommandLine::Get(), TEXT("ResetTutorials"));
 }
diff --git a/Engine/Source/Editor/LevelEditor/Private/LevelEditor.cpp b/Engine/Source/Editor/LevelEditor/Private/LevelEditor.cpp
--- a/Engine/Source/Editor/LevelEditor/Private/LevelEditor.cpp
+++ b/Engine/Source/Editor/LevelEditor/Private/LevelEditor.cpp
@@ -651,8 +651,6 @@ void FLevelEditorModule::BindGlobalLevelEditorCommands()
 	FCanExecuteAction DefaultExecuteAction = FCanExecuteAction::CreateStatic( &FLevelEditorActionCallbacks::DefaultCanExecuteAction );
 
 	ActionList.MapAction( Commands.BrowseDocumentation, FExecuteAction::CreateStatic( &FLevelEditorActionCallbacks::BrowseDocumentation ) );
-	ActionList.MapAction( Commands.BrowseAPIReference, FExecuteAction::CreateStatic( &FLevelEditorActionCallbacks::BrowseAPIReference ) );
-	ActionList.MapAction( Commands.BrowseCVars, FExecuteAction::CreateStatic( &FLevelEditorActionCallbacks::BrowseCVars ) );
 	ActionList.MapAction( Commands.BrowseViewportControls, FExecuteAction::CreateStatic( &FLevelEditorActionCallbacks::BrowseViewportControls ) );
 	ActionList.MapAction( Commands.NewLevel, FExecuteAction::CreateStatic( &FLevelEditorActionCallbacks::NewLevel ), FCanExecuteAction::CreateStatic( &FLevelEditorActionCallbacks::NewLevel_CanExecute ) );
 	ActionList.MapAction(Commands.OpenLevel, FExecuteAction::CreateStatic(&FLevelEditorActionCallbacks::OpenLevel), FCanExecuteAction::CreateStatic(&FLevelEditorActionCallbacks::OpenLevel_CanExecute));
diff --git a/Engine/Source/Editor/LevelEditor/Private/LevelEditorActions.cpp b/Engine/Source/Editor/LevelEditor/Private/LevelEditorActions.cpp
--- a/Engine/Source/Editor/LevelEditor/Private/LevelEditorActions.cpp
+++ b/Engine/Source/Editor/LevelEditor/Private/LevelEditorActions.cpp
@@ -181,16 +181,6 @@ void FLevelEditorActionCallbacks::BrowseDocumentation()
 	IDocumentation::Get()->OpenHome(FDocumentationSourceInfo(TEXT("help_menu")));
 }
 
-void FLevelEditorActionCallbacks::BrowseAPIReference()
-{
-	IDocumentation::Get()->OpenAPIHome(FDocumentationSourceInfo(TEXT("help_menu")));
-}
-
-void FLevelEditorActionCallbacks::BrowseCVars()
-{
-	GEditor->Exec(GetWorld(), TEXT("help"));
-}
-
 void FLevelEditorActionCallbacks::BrowseViewportControls()
 {
 	FString URL;
@@ -3329,10 +3319,8 @@ FLevelEditorCommands::FLevelEditorCommands()
 PRAGMA_DISABLE_OPTIMIZATION
 void FLevelEditorCommands::RegisterCommands()
 {
-	UI_COMMAND( BrowseDocumentation, "Documentation...", "Opens the main documentation page, and allows you to search across all UE support sites.", EUserInterfaceActionType::Button, FInputChord( EKeys::F1 ) );
-	UI_COMMAND( BrowseAPIReference, "API Reference...", "Opens the API reference documentation", EUserInterfaceActionType::Button, FInputChord() );
-	UI_COMMAND( BrowseCVars, "Console Variables", "Creates an HTML file to browse the console variables and commands (console command 'help')", EUserInterfaceActionType::Button, FInputChord() );
-	UI_COMMAND( BrowseViewportControls, "Viewport Controls...", "Opens the viewport controls cheat sheet", EUserInterfaceActionType::Button, FInputChord() );
+	UI_COMMAND( BrowseDocumentation, "Level Editor Documentation", "Details on how to use the Level Editor", EUserInterfaceActionType::Button, FInputChord( EKeys::F1 ) );
+	UI_COMMAND( BrowseViewportControls, "Viewport Controls", "Ways to move around in the 3D viewport", EUserInterfaceActionType::Button, FInputChord() );
 	UI_COMMAND( NewLevel, "New Level...", "Create a new level, or choose a level template to start from.", EUserInterfaceActionType::Button, FInputChord( EModifierKey::Control, EKeys::N ) );
 	UI_COMMAND( OpenLevel, "Open Level...", "Loads an existing level", EUserInterfaceActionType::Button, FInputChord( EModifierKey::Control, EKeys::O ) );
 	UI_COMMAND( Save, "Save Current Level", "Saves the current level to disk", EUserInterfaceActionType::Button, FInputChord(EModifierKey::Control, EKeys::S) );
diff --git a/Engine/Source/Editor/LevelEditor/Private/LevelEditorMenu.cpp b/Engine/Source/Editor/LevelEditor/Private/LevelEditorMenu.cpp
--- a/Engine/Source/Editor/LevelEditor/Private/LevelEditorMenu.cpp
+++ b/Engine/Source/Editor/LevelEditor/Private/LevelEditorMenu.cpp
@@ -213,16 +213,10 @@ void FLevelEditorMenu::RegisterLevelEditorMenus()
 		static void ExtendHelpMenu()
 		{
 			UToolMenu* Menu = UToolMenus::Get()->RegisterMenu("LevelEditor.MainMenu.Help", "MainFrame.MainMenu.Help");
-			FToolMenuSection& Section = Menu->AddSection("HelpBrowse", NSLOCTEXT("MainHelpMenu", "Browse", "Browse"), FToolMenuInsert("BugReporting", EToolMenuInsertType::Before));
+			FToolMenuSection& Section = Menu->AddSection("HelpResources", NSLOCTEXT("MainHelpMenu", "LevelEditorHelpResources", "Level Editor Resources"), FToolMenuInsert("Learn", EToolMenuInsertType::First));
 			{
 				Section.AddMenuEntry( FLevelEditorCommands::Get().BrowseDocumentation );
 
-				Section.AddMenuEntry( FLevelEditorCommands::Get().BrowseAPIReference );
-
-				Section.AddMenuEntry( FLevelEditorCommands::Get().BrowseCVars );
-
-				Section.AddSeparator( "HelpBrowse" );
-
 				Section.AddMenuEntry( FLevelEditorCommands::Get().BrowseViewportControls );
 			}
 		}
diff --git a/Engine/Source/Editor/LevelEditor/Public/LevelEditorActions.h b/Engine/Source/Editor/LevelEditor/Public/LevelEditorActions.h
--- a/Engine/Source/Editor/LevelEditor/Public/LevelEditorActions.h
+++ b/Engine/Source/Editor/LevelEditor/Public/LevelEditorActions.h
@@ -40,8 +40,6 @@ public:
 public:
 	
 	TSharedPtr< FUICommandInfo > BrowseDocumentation;
-	TSharedPtr< FUICommandInfo > BrowseAPIReference;
-	TSharedPtr< FUICommandInfo > BrowseCVars;
 	TSharedPtr< FUICommandInfo > BrowseViewportControls;
 
 	/** Level file commands */
@@ -641,12 +639,6 @@ public:
 	/** Opens the global documentation homepage */
 	static void BrowseDocumentation();
 
-	/** Opens the API documentation CHM */
-	static void BrowseAPIReference();
-
-	/** Creates an HTML file to browse the console variables and commands */
-	static void BrowseCVars();
-
 	/** Opens the viewport controls page*/
 	static void BrowseViewportControls();
 
diff --git a/Engine/Source/Editor/MainFrame/Private/Frame/MainFrameActions.cpp b/Engine/Source/Editor/MainFrame/Private/Frame/MainFrameActions.cpp
--- a/Engine/Source/Editor/MainFrame/Private/Frame/MainFrameActions.cpp
+++ b/Engine/Source/Editor/MainFrame/Private/Frame/MainFrameActions.cpp
@@ -46,6 +46,7 @@
 #include "IUATHelperModule.h"
 #include "Menus/LayoutsMenu.h"
 #include "TargetReceipt.h"
+#include "IDocumentation.h"
 
 #include "Settings/EditorSettings.h"
 #include "AnalyticsEventAttribute.h"
@@ -179,36 +180,42 @@ void FMainFrameCommands::RegisterCommands()
 	UI_COMMAND(OpenMarketplace, "Open Marketplace", "Opens the Marketplace", EUserInterfaceActionType::Button, FInputChord());
 	ActionList->MapAction(OpenMarketplace, FExecuteAction::CreateStatic(&FMainFrameActionCallbacks::OpenMarketplace));
 
-	UI_COMMAND(VisitOnlineLearning, "Online Learning...", "Learn Unreal Engine for free with easy-to-follow video courses and guided learning paths.", EUserInterfaceActionType::Button, FInputChord());
+	UI_COMMAND(DocumentationHome, "Documentation Home", "Authoritative, in-depth technical resources for using Unreal Engine", EUserInterfaceActionType::Button, FInputChord());
+	ActionList->MapAction(DocumentationHome, FExecuteAction::CreateStatic(&FMainFrameActionCallbacks::DocumentationHome));
+
+	UI_COMMAND(VisitOnlineLearning, "Online Learning", "Learn Unreal Engine for free with easy-to-follow video courses and guided learning paths", EUserInterfaceActionType::Button, FInputChord());
 	ActionList->MapAction(VisitOnlineLearning, FExecuteAction::CreateStatic(&FMainFrameActionCallbacks::VisitOnlineLearning));
 
-	UI_COMMAND(VisitForums, "Forums...", "Go to the Unreal Engine forums to view announcements and engage in discussions with other developers.", EUserInterfaceActionType::Button, FInputChord());
-	ActionList->MapAction(VisitForums, FExecuteAction::CreateStatic(&FMainFrameActionCallbacks::VisitForums));
+	UI_COMMAND(BrowseAPIReference, "C++ API Reference", "Classes, functions, and other elements that make up the C++ API", EUserInterfaceActionType::Button, FInputChord());
+	ActionList->MapAction(BrowseAPIReference, FExecuteAction::CreateStatic(&FMainFrameActionCallbacks::BrowseAPIReference));
 
-	UI_COMMAND(ReportABug, "Report a Bug...", "Found a bug?  Go here to fill out a bug report", EUserInterfaceActionType::Button, FInputChord());
-	ActionList->MapAction(ReportABug, FExecuteAction::CreateStatic(&FMainFrameActionCallbacks::ReportABug));
+	UI_COMMAND(BrowseCVars, "Console Variables", "Reference companion for console variables and commands", EUserInterfaceActionType::Button, FInputChord());
+	ActionList->MapAction(BrowseCVars, FExecuteAction::CreateStatic(&FMainFrameActionCallbacks::BrowseCVars));
 
-	UI_COMMAND(OpenIssueTracker, "Issue Tracker", "Go here to view the Unreal Engine bug tracking website", EUserInterfaceActionType::Button, FInputChord());
-	ActionList->MapAction(OpenIssueTracker, FExecuteAction::CreateStatic(&FMainFrameActionCallbacks::OpenIssueTracker));
+	UI_COMMAND(VisitForums, "Forums", "View announcements and engage in discussions with other developers", EUserInterfaceActionType::Button, FInputChord());
+	ActionList->MapAction(VisitForums, FExecuteAction::CreateStatic(&FMainFrameActionCallbacks::VisitForums));
 
-	UI_COMMAND( VisitAskAQuestionPage, "Ask a Question...", "Have a question?  Go here to ask about anything and everything related to Unreal.", EUserInterfaceActionType::Button, FInputChord() );
-	ActionList->MapAction( VisitAskAQuestionPage, FExecuteAction::CreateStatic( &FMainFrameActionCallbacks::VisitAskAQuestionPage ) );
+	UI_COMMAND(VisitSearchForAnswersPage, "AnswerHub", "Search for answers, ask questions, and share your knowledge with other developers", EUserInterfaceActionType::Button, FInputChord());
+	ActionList->MapAction(VisitSearchForAnswersPage, FExecuteAction::CreateStatic(&FMainFrameActionCallbacks::VisitSearchForAnswersPage));
 
-	UI_COMMAND( VisitSearchForAnswersPage, "Answer Hub...", "Go to the AnswerHub to ask questions, search existing answers, and share your knowledge with other UE developers.", EUserInterfaceActionType::Button, FInputChord() );
-	ActionList->MapAction( VisitSearchForAnswersPage, FExecuteAction::CreateStatic( &FMainFrameActionCallbacks::VisitSearchForAnswersPage ) );
+	UI_COMMAND(VisitSupportWebSite, "Support", "Options for personalized technical support", EUserInterfaceActionType::Button, FInputChord());
+	ActionList->MapAction(VisitSupportWebSite, FExecuteAction::CreateStatic(&FMainFrameActionCallbacks::VisitSupportWebSite));
 
-	UI_COMMAND( VisitSupportWebSite, "Support...", "Navigates to the Unreal Engine Support web site's main page.", EUserInterfaceActionType::Button, FInputChord() );
-	ActionList->MapAction( VisitSupportWebSite, FExecuteAction::CreateStatic( &FMainFrameActionCallbacks::VisitSupportWebSite ) );
+	UI_COMMAND(ReportABug, "Report a Bug", "Found a bug? Let us know about it", EUserInterfaceActionType::Button, FInputChord());
+	ActionList->MapAction(ReportABug, FExecuteAction::CreateStatic(&FMainFrameActionCallbacks::ReportABug));
 
-	UI_COMMAND( VisitEpicGamesDotCom, "Visit UnrealEngine.com...", "Navigates to UnrealEngine.com where you can learn more about Unreal Technology.", EUserInterfaceActionType::Button, FInputChord() );
-	ActionList->MapAction( VisitEpicGamesDotCom, FExecuteAction::CreateStatic( &FMainFrameActionCallbacks::VisitEpicGamesDotCom ) );
+	UI_COMMAND(OpenIssueTracker, "Issue Tracker", "Check the current status of public bugs and other issues", EUserInterfaceActionType::Button, FInputChord());
+	ActionList->MapAction(OpenIssueTracker, FExecuteAction::CreateStatic(&FMainFrameActionCallbacks::OpenIssueTracker));
 
-	UI_COMMAND( AboutUnrealEd, "About Editor...", "Displays application credits and copyright information", EUserInterfaceActionType::Button, FInputChord() );
+	UI_COMMAND( AboutUnrealEd, "About Unreal Editor", "Version and copyright information", EUserInterfaceActionType::Button, FInputChord() );
 	ActionList->MapAction( AboutUnrealEd, FExecuteAction::CreateStatic( &FMainFrameActionCallbacks::AboutUnrealEd_Execute ) );
 
-	UI_COMMAND( CreditsUnrealEd, "Credits", "Displays application credits", EUserInterfaceActionType::Button, FInputChord() );
+	UI_COMMAND( CreditsUnrealEd, "Credits", "Contributors to this version of Unreal Engine", EUserInterfaceActionType::Button, FInputChord() );
 	ActionList->MapAction( CreditsUnrealEd, FExecuteAction::CreateStatic(&FMainFrameActionCallbacks::CreditsUnrealEd_Execute) );
 
+	UI_COMMAND(VisitEpicGamesDotCom, "Visit UnrealEngine.com", "Learn more about Unreal technology", EUserInterfaceActionType::Button, FInputChord());
+	ActionList->MapAction(VisitEpicGamesDotCom, FExecuteAction::CreateStatic(&FMainFrameActionCallbacks::VisitEpicGamesDotCom));
+
 	// Layout commands
 	UI_COMMAND(ImportLayout, "Import Layout...", "Import a custom layout (or set of layouts) from a different directory and load it into your current instance of the Unreal Editor UI", EUserInterfaceActionType::Button, FInputChord());
 	ActionList->MapAction(ImportLayout, FExecuteAction::CreateStatic(&FLayoutsMenuLoad::ImportLayout));
@@ -575,16 +582,6 @@ void FMainFrameActionCallbacks::OpenIssueTracker()
 	}
 }
 
-void FMainFrameActionCallbacks::VisitAskAQuestionPage()
-{
-	FString AskAQuestionURL;
-	if(FUnrealEdMisc::Get().GetURL( TEXT("AskAQuestionURL"), AskAQuestionURL, true ))
-	{
-		FPlatformProcess::LaunchURL( *AskAQuestionURL, NULL, NULL );
-	}
-}
-
-
 void FMainFrameActionCallbacks::VisitSearchForAnswersPage()
 {
 	FString SearchForAnswersURL;
@@ -594,7 +591,6 @@ void FMainFrameActionCallbacks::VisitSearchForAnswersPage()
 	}
 }
 
-
 void FMainFrameActionCallbacks::VisitSupportWebSite()
 {
 	FString SupportWebsiteURL;
@@ -604,7 +600,6 @@ void FMainFrameActionCallbacks::VisitSupportWebSite()
 	}
 }
 
-
 void FMainFrameActionCallbacks::VisitEpicGamesDotCom()
 {
 	FString EpicGamesURL;
@@ -614,6 +609,11 @@ void FMainFrameActionCallbacks::VisitEpicGamesDotCom()
 	}
 }
 
+void FMainFrameActionCallbacks::DocumentationHome()
+{
+	IDocumentation::Get()->OpenHome(FDocumentationSourceInfo(TEXT("help_menu")));
+}
+
 void FMainFrameActionCallbacks::VisitOnlineLearning()
 {
 	FString URL;
@@ -623,6 +623,16 @@ void FMainFrameActionCallbacks::VisitOnlineLearning()
 	}
 }
 
+void FMainFrameActionCallbacks::BrowseAPIReference()
+{
+	IDocumentation::Get()->OpenAPIHome(FDocumentationSourceInfo(TEXT("help_menu")));
+}
+
+void FMainFrameActionCallbacks::BrowseCVars()
+{
+	GEditor->Exec(GEditor->GetEditorWorldContext().World(), TEXT("help"));
+}
+
 void FMainFrameActionCallbacks::VisitForums()
 {
 	FString URL;
diff --git a/Engine/Source/Editor/MainFrame/Private/Frame/MainFrameActions.h b/Engine/Source/Editor/MainFrame/Private/Frame/MainFrameActions.h
--- a/Engine/Source/Editor/MainFrame/Private/Frame/MainFrameActions.h
+++ b/Engine/Source/Editor/MainFrame/Private/Frame/MainFrameActions.h
@@ -52,9 +52,11 @@ public:
 	TSharedPtr< FUICommandInfo > OpenWidgetReflector;
 	TSharedPtr< FUICommandInfo > OpenMarketplace;
 
+	TSharedPtr< FUICommandInfo > DocumentationHome;
 	TSharedPtr< FUICommandInfo > VisitOnlineLearning;
+	TSharedPtr< FUICommandInfo > BrowseAPIReference;
+	TSharedPtr< FUICommandInfo > BrowseCVars;
 	TSharedPtr< FUICommandInfo > VisitForums;
-	TSharedPtr< FUICommandInfo > VisitAskAQuestionPage;
 	TSharedPtr< FUICommandInfo > VisitSearchForAnswersPage;
 	TSharedPtr< FUICommandInfo > ReportABug;
 	TSharedPtr< FUICommandInfo > OpenIssueTracker;
@@ -147,9 +149,6 @@ public:
 	/** Opens the issue tracker page */
 	static void OpenIssueTracker();
 
-	/** Visits the "ask a question" page on UDN */
-	static void VisitAskAQuestionPage();
-
 	/** Visits the "search for answers" page on UDN */
 	static void VisitSearchForAnswersPage();
 
@@ -159,9 +158,18 @@ public:
 	/** Visits EpicGames.com */
 	static void VisitEpicGamesDotCom();
 
+	/** Opens the documentation home page*/
+	static void DocumentationHome();
+
 	/** Visits The Unreal Online Learning page*/
 	static void VisitOnlineLearning();
 
+	/** Opens the API documentation site */
+	static void BrowseAPIReference();
+
+	/** Creates an HTML file to browse the console variables and commands */
+	static void BrowseCVars();
+
 	static void VisitForums();
 
 	static void AboutUnrealEd_Execute();
diff --git a/Engine/Source/Editor/MainFrame/Private/Menus/MainMenu.cpp b/Engine/Source/Editor/MainFrame/Private/Menus/MainMenu.cpp
--- a/Engine/Source/Editor/MainFrame/Private/Menus/MainMenu.cpp
+++ b/Engine/Source/Editor/MainFrame/Private/Menus/MainMenu.cpp
@@ -260,38 +260,43 @@ void FMainMenu::RegisterWindowMenu()
 void FMainMenu::RegisterHelpMenu()
 {
 	UToolMenu* Menu = UToolMenus::Get()->RegisterMenu("MainFrame.MainMenu.Help");
-	FToolMenuSection& BugReportingSection = Menu->AddSection("BugReporting", NSLOCTEXT("MainHelpMenu", "BugsReporting", "Bugs"));
+
+	FToolMenuSection& LearnSection = Menu->AddSection("Learn", NSLOCTEXT("MainHelpMenu", "LearnSection", "Learn"));
 	{
-		BugReportingSection.AddMenuEntry(FMainFrameCommands::Get().ReportABug);
-		BugReportingSection.AddMenuEntry(FMainFrameCommands::Get().OpenIssueTracker);
+		LearnSection.AddMenuEntry(FMainFrameCommands::Get().DocumentationHome);
+		LearnSection.AddMenuEntry(FMainFrameCommands::Get().VisitOnlineLearning);
 	}
 
-	FToolMenuSection& HelpOnlineSection = Menu->AddSection("HelpOnline", NSLOCTEXT("MainHelpMenu", "Online", "Help Online"));
+	FToolMenuSection& ReferenceSection = Menu->AddSection("Reference", NSLOCTEXT("MainHelpMenu", "ReferenceSection", "Reference"));
+	{
+		ReferenceSection.AddMenuEntry(FMainFrameCommands::Get().BrowseAPIReference);
+		ReferenceSection.AddMenuEntry(FMainFrameCommands::Get().BrowseCVars);
+	}
+
+	FToolMenuSection& HelpOnlineSection = Menu->AddSection("HelpOnline", NSLOCTEXT("MainHelpMenu", "HelpOnlineSection", "Help"));
 	{
-		HelpOnlineSection.AddMenuEntry(FMainFrameCommands::Get().VisitSupportWebSite);
 		HelpOnlineSection.AddMenuEntry(FMainFrameCommands::Get().VisitForums);
 		HelpOnlineSection.AddMenuEntry(FMainFrameCommands::Get().VisitSearchForAnswersPage);
-		HelpOnlineSection.AddMenuEntry(FMainFrameCommands::Get().VisitOnlineLearning);
-
-
-
-		const FText SupportWebSiteLabel = NSLOCTEXT("MainHelpMenu", "VisitUnrealEngineSupportWebSite", "Unreal Engine Support Web Site...");
-
-		HelpOnlineSection.AddSeparator("EpicGamesHelp");
-		HelpOnlineSection.AddMenuEntry(FMainFrameCommands::Get().VisitEpicGamesDotCom);
+		HelpOnlineSection.AddMenuEntry(FMainFrameCommands::Get().VisitSupportWebSite);
+	}
 
-		HelpOnlineSection.AddSeparator("Credits");
-		HelpOnlineSection.AddMenuEntry(FMainFrameCommands::Get().CreditsUnrealEd);
+	FToolMenuSection& BugReportingSection = Menu->AddSection("BugReporting", NSLOCTEXT("MainHelpMenu", "BugsReporting", "Bugs"));
+	{
+		BugReportingSection.AddMenuEntry(FMainFrameCommands::Get().ReportABug);
+		BugReportingSection.AddMenuEntry(FMainFrameCommands::Get().OpenIssueTracker);
 	}
 
-#if !PLATFORM_MAC // Handled by app's menu in menu bar
 	FToolMenuSection& HelpApplicationSection = Menu->AddSection("HelpApplication", NSLOCTEXT("MainHelpMenu", "Application", "Application"));
 	{
-		const FText AboutWindowTitle = NSLOCTEXT("MainHelpMenu", "AboutUnrealEditor", "About Unreal Editor...");
 
-		HelpApplicationSection.AddMenuEntry(FMainFrameCommands::Get().AboutUnrealEd, AboutWindowTitle);
-	}
+#if !PLATFORM_MAC // Handled by app's menu in menu bar
+		HelpApplicationSection.AddMenuEntry(FMainFrameCommands::Get().AboutUnrealEd);
 #endif
+
+		HelpApplicationSection.AddMenuEntry(FMainFrameCommands::Get().CreditsUnrealEd);
+		HelpApplicationSection.AddSeparator("EpicGamesHelp");
+		HelpApplicationSection.AddMenuEntry(FMainFrameCommands::Get().VisitEpicGamesDotCom);
+	}
 }
 
 #undef LOCTEXT_NAMESPACE
diff --git a/Engine/Source/Editor/UnrealEd/Private/Toolkits/AssetEditorToolkit.cpp b/Engine/Source/Editor/UnrealEd/Private/Toolkits/AssetEditorToolkit.cpp
--- a/Engine/Source/Editor/UnrealEd/Private/Toolkits/AssetEditorToolkit.cpp
+++ b/Engine/Source/Editor/UnrealEd/Private/Toolkits/AssetEditorToolkit.cpp
@@ -982,11 +982,16 @@ void FAssetEditorToolkit::FillDefaultAssetMenuCommands(FToolMenuSection& InSecti
 
 void FAssetEditorToolkit::FillDefaultHelpMenuCommands(FToolMenuSection& InSection)
 {
-	FFormatNamedArguments Args;
-	Args.Add(TEXT("Editor"), GetBaseToolkitName());
-	const FText ToolTip = FText::Format(LOCTEXT("BrowseDocumentationTooltip", "Browse {Editor} documentation..."), Args);
+	// Only show the documentation menu item if the asset editor has specified a resource.
+	FString DocLink = GetDocumentationLink();
+	if (DocLink != "%ROOT%") 
+	{
+		FFormatNamedArguments Args;
+		Args.Add(TEXT("Editor"), GetBaseToolkitName());
 
-	InSection.AddMenuEntry(FGlobalEditorCommonCommands::Get().OpenDocumentation, ToolTip);
+		const FText ToolTip = FText::Format(LOCTEXT("BrowseDocumentationTooltip", "Details on using the {Editor}"), Args);
+		InSection.AddMenuEntry(FGlobalEditorCommonCommands::Get().OpenDocumentation, FText::Format(LOCTEXT("AssetEditorDocumentationMenuLabel", "{Editor} Documentation"), Args), ToolTip);
+	}
 }
 
 FName FAssetEditorToolkit::GetToolMenuAppName() const
diff --git a/Engine/Source/Editor/UnrealEd/Private/Toolkits/SStandaloneAssetEditorToolkitHost.cpp b/Engine/Source/Editor/UnrealEd/Private/Toolkits/SStandaloneAssetEditorToolkitHost.cpp
--- a/Engine/Source/Editor/UnrealEd/Private/Toolkits/SStandaloneAssetEditorToolkitHost.cpp
+++ b/Engine/Source/Editor/UnrealEd/Private/Toolkits/SStandaloneAssetEditorToolkitHost.cpp
@@ -127,11 +127,14 @@ void SStandaloneAssetEditorToolkitHost::CreateDefaultStandaloneMenuBar(UToolMenu
 		{
 			const FName MenuName = *(InMenuBar->GetMenuName().ToString() + TEXT(".") + TEXT("Help"));
 			UToolMenu* Menu = UToolMenus::Get()->ExtendMenu(MenuName);
-			FToolMenuSection& Section = Menu->AddSection("HelpBrowse", NSLOCTEXT("MainHelpMenu", "Browse", "Browse"));
-			Section.InsertPosition = FToolMenuInsert("HelpOnline", EToolMenuInsertType::Before);
-			Section.AddDynamicEntry(NAME_None, FNewToolMenuSectionDelegate::CreateLambda([](FToolMenuSection& InSection)
+			Menu->AddDynamicSection(NAME_None, FNewToolMenuDelegate::CreateLambda([](UToolMenu* InMenu)
 			{
-				InSection.FindContext<UAssetEditorToolkitMenuContext>()->Toolkit.Pin()->FillDefaultHelpMenuCommands(InSection);
+				TSharedPtr<FAssetEditorToolkit> Toolkit = InMenu->FindContext<UAssetEditorToolkitMenuContext>()->Toolkit.Pin();
+				FFormatNamedArguments Args;
+				Args.Add(TEXT("Editor"), Toolkit->GetBaseToolkitName());
+				FToolMenuSection& Section = InMenu->AddSection("HelpResources", FText::Format(NSLOCTEXT("MainHelpMenu", "AssetEditorHelpResources", "{Editor} Resources"), Args));
+				Section.InsertPosition = FToolMenuInsert("Learn", EToolMenuInsertType::First);
+				Toolkit->FillDefaultHelpMenuCommands(Section);
 			}));
 		}
 	};
