Fix default editor layouts not appearing when using pak files

diff --git a/Engine/Source/Editor/MainFrame/Private/Menus/LayoutsMenu.cpp b/Engine/Source/Editor/MainFrame/Private/Menus/LayoutsMenu.cpp
--- a/Engine/Source/Editor/MainFrame/Private/Menus/LayoutsMenu.cpp
+++ b/Engine/Source/Editor/MainFrame/Private/Menus/LayoutsMenu.cpp
@@ -358,8 +358,14 @@ void FPrivateLayoutsMenu::DisplayLayoutsInternal(FToolMenuSection& InSection, co
 	// If there are Layout ini files, read them
 	for (int32 LayoutIndex = 0; LayoutIndex < InLayoutIniFileNames.Num(); ++LayoutIndex)
 	{
-		const FString LayoutFilePath = IFileManager::Get().ConvertToAbsolutePathForExternalAppForRead(
+		FString LayoutFilePath = IFileManager::Get().ConvertToAbsolutePathForExternalAppForRead(
 			*FPaths::Combine(InLayoutsDirectory, InLayoutIniFileNames[LayoutIndex]));
+
+		if (LayoutFilePath.StartsWith(TEXT("Pak: ")))
+		{
+			LayoutFilePath = FPaths::Combine(InLayoutsDirectory, InLayoutIniFileNames[LayoutIndex]);
+		}
+
 		// Make sure it is a layout file
 		GConfig->UnloadFile(LayoutFilePath); // We must re-read it to avoid the Editor to use a previously cached name and description
 		if (FLayoutSaveRestore::IsValidConfig(LayoutFilePath))
