Help menu changes for community: reorganizations, icon changes

diff --git a/Engine/Config/BaseEditor.ini b/Engine/Config/BaseEditor.ini
--- a/Engine/Config/BaseEditor.ini
+++ b/Engine/Config/BaseEditor.ini
@@ -279,7 +279,7 @@ DocumentationURL="https://docs.unrealengine.com/{VERSION}/{I18N}/"
 OnlineLearningURL="https://www.unrealengine.com/onlinelearning-courses"
 ForumsURL="http://forums.unrealengine.com"
 APIDocsURL="https://docs.unrealengine.com/{VERSION}/{I18N}/API/"
-SearchForAnswersURL="http://answers.unrealengine.com"
+SearchForAnswersURL="https://forums.unrealengine.com/tag/question"
 SupportWebsiteURL="http://unrealengine.com/support"
 SourceCodeIDEURL_Windows="https://aka.ms/vs/16/release/vs_Community.exe"
 SourceCodeIDEURL_Mac="https://developer.apple.com/xcode"
diff --git a/Engine/Source/Editor/EditorStyle/Private/StarshipStyle.cpp b/Engine/Source/Editor/EditorStyle/Private/StarshipStyle.cpp
--- a/Engine/Source/Editor/EditorStyle/Private/StarshipStyle.cpp
+++ b/Engine/Source/Editor/EditorStyle/Private/StarshipStyle.cpp
@@ -4306,11 +4306,11 @@ void FStarshipEditorStyle::FStyle::SetupLevelEditorStyle()
 
 		Set("SystemWideCommands.OpenDocumentation", new IMAGE_BRUSH_SVG("Starship/Common/Documentation", Icon16x16));
 		Set("MainFrame.DocumentationHome",	        new IMAGE_BRUSH_SVG("Starship/Common/Documentation", Icon16x16));
-		Set("MainFrame.VisitOnlineLearning",	    new IMAGE_BRUSH_SVG("Starship/Common/Documentation", Icon16x16));
 		Set("MainFrame.BrowseAPIReference",         new IMAGE_BRUSH_SVG("Starship/Common/Documentation", Icon16x16));
 		Set("MainFrame.BrowseCVars",                new IMAGE_BRUSH_SVG("Starship/Common/Console", Icon16x16));
-		Set("MainFrame.VisitForums",                new IMAGE_BRUSH_SVG("Starship/Common/WebBrowser", Icon16x16));
-		Set("MainFrame.VisitSearchForAnswersPage",  new IMAGE_BRUSH_SVG("Starship/Common/WebBrowser", Icon16x16));
+		Set("MainFrame.VisitOnlineLearning",		new IMAGE_BRUSH_SVG("Starship/Common/Tutorials", Icon16x16));
+		Set("MainFrame.VisitForums",                new IMAGE_BRUSH_SVG("Starship/Common/Forums", Icon16x16));
+		Set("MainFrame.VisitSearchForAnswersPage",  new IMAGE_BRUSH_SVG("Starship/Common/QuestionAnswer", Icon16x16));
 		Set("MainFrame.VisitSupportWebSite",        new IMAGE_BRUSH_SVG("Starship/Common/Support", Icon16x16));
 		Set("MainFrame.VisitEpicGamesDotCom",       new IMAGE_BRUSH_SVG("About/EpicGamesLogo", Icon16x16));
 		Set("MainFrame.AboutUnrealEd",              new IMAGE_BRUSH_SVG("About/UnrealLogo", Icon16x16));
diff --git a/Engine/Source/Editor/MainFrame/Private/Frame/MainFrameActions.cpp b/Engine/Source/Editor/MainFrame/Private/Frame/MainFrameActions.cpp
--- a/Engine/Source/Editor/MainFrame/Private/Frame/MainFrameActions.cpp
+++ b/Engine/Source/Editor/MainFrame/Private/Frame/MainFrameActions.cpp
@@ -195,7 +195,7 @@ void FMainFrameCommands::RegisterCommands()
 	UI_COMMAND(VisitForums, "Forums", "View announcements and engage in discussions with other developers", EUserInterfaceActionType::Button, FInputChord());
 	ActionList->MapAction(VisitForums, FExecuteAction::CreateStatic(&FMainFrameActionCallbacks::VisitForums));
 
-	UI_COMMAND(VisitSearchForAnswersPage, "AnswerHub", "Search for answers, ask questions, and share your knowledge with other developers", EUserInterfaceActionType::Button, FInputChord());
+	UI_COMMAND(VisitSearchForAnswersPage, "Q&A", "Search for answers, ask questions, and share your knowledge with other developers", EUserInterfaceActionType::Button, FInputChord());
 	ActionList->MapAction(VisitSearchForAnswersPage, FExecuteAction::CreateStatic(&FMainFrameActionCallbacks::VisitSearchForAnswersPage));
 
 	UI_COMMAND(VisitSupportWebSite, "Support", "Options for personalized technical support", EUserInterfaceActionType::Button, FInputChord());
diff --git a/Engine/Source/Editor/MainFrame/Private/Menus/MainMenu.cpp b/Engine/Source/Editor/MainFrame/Private/Menus/MainMenu.cpp
--- a/Engine/Source/Editor/MainFrame/Private/Menus/MainMenu.cpp
+++ b/Engine/Source/Editor/MainFrame/Private/Menus/MainMenu.cpp
@@ -261,27 +261,23 @@ void FMainMenu::RegisterHelpMenu()
 {
 	UToolMenu* Menu = UToolMenus::Get()->RegisterMenu("MainFrame.MainMenu.Help");
 
-	FToolMenuSection& LearnSection = Menu->AddSection("Learn", NSLOCTEXT("MainHelpMenu", "LearnSection", "Learn"));
-	{
-		LearnSection.AddMenuEntry(FMainFrameCommands::Get().DocumentationHome);
-		LearnSection.AddMenuEntry(FMainFrameCommands::Get().VisitOnlineLearning);
-	}
-
 	FToolMenuSection& ReferenceSection = Menu->AddSection("Reference", NSLOCTEXT("MainHelpMenu", "ReferenceSection", "Reference"));
 	{
+		ReferenceSection.AddMenuEntry(FMainFrameCommands::Get().DocumentationHome);
 		ReferenceSection.AddMenuEntry(FMainFrameCommands::Get().BrowseAPIReference);
 		ReferenceSection.AddMenuEntry(FMainFrameCommands::Get().BrowseCVars);
 	}
 
-	FToolMenuSection& HelpOnlineSection = Menu->AddSection("HelpOnline", NSLOCTEXT("MainHelpMenu", "HelpOnlineSection", "Help"));
+	FToolMenuSection& CommunitySection = Menu->AddSection("Community", NSLOCTEXT("MainHelpMenu", "CommunitySection", "Community"));
 	{
-		HelpOnlineSection.AddMenuEntry(FMainFrameCommands::Get().VisitForums);
-		HelpOnlineSection.AddMenuEntry(FMainFrameCommands::Get().VisitSearchForAnswersPage);
-		HelpOnlineSection.AddMenuEntry(FMainFrameCommands::Get().VisitSupportWebSite);
+		CommunitySection.AddMenuEntry(FMainFrameCommands::Get().VisitOnlineLearning);
+		CommunitySection.AddMenuEntry(FMainFrameCommands::Get().VisitForums);
+		CommunitySection.AddMenuEntry(FMainFrameCommands::Get().VisitSearchForAnswersPage);
 	}
 
-	FToolMenuSection& BugReportingSection = Menu->AddSection("BugReporting", NSLOCTEXT("MainHelpMenu", "BugsReporting", "Bugs"));
+	FToolMenuSection& BugReportingSection = Menu->AddSection("Support", NSLOCTEXT("MainHelpMenu", "SupportSection", "Support"));
 	{
+		BugReportingSection.AddMenuEntry(FMainFrameCommands::Get().VisitSupportWebSite);
 		BugReportingSection.AddMenuEntry(FMainFrameCommands::Get().ReportABug);
 		BugReportingSection.AddMenuEntry(FMainFrameCommands::Get().OpenIssueTracker);
 	}
