Hide some more C++ creation/editing menus when !GUnrealEd->GetUnrealEdOptions()->IsCPPAllowed() (follow up on 18815523)

diff --git a/Engine/Plugins/Editor/ContentBrowser/ContentBrowserClassDataSource/Source/ContentBrowserClassDataSource/Private/ContentBrowserClassDataSource.cpp b/Engine/Plugins/Editor/ContentBrowser/ContentBrowserClassDataSource/Source/ContentBrowserClassDataSource/Private/ContentBrowserClassDataSource.cpp
--- a/Engine/Plugins/Editor/ContentBrowser/ContentBrowserClassDataSource/Source/ContentBrowserClassDataSource/Private/ContentBrowserClassDataSource.cpp
+++ b/Engine/Plugins/Editor/ContentBrowser/ContentBrowserClassDataSource/Source/ContentBrowserClassDataSource/Private/ContentBrowserClassDataSource.cpp
@@ -17,6 +17,9 @@
 #include "ContentBrowserDataSubsystem.h"
 #include "ContentBrowserDataUtils.h"
 #include "ContentBrowserItemPath.h"
+#include "Editor/UnrealEdEngine.h"
+#include "Preferences/UnrealEdOptions.h"
+#include "UnrealEdGlobals.h"
 
 #define LOCTEXT_NAMESPACE "ContentBrowserClassDataSource"
 
@@ -612,6 +615,11 @@ TSharedPtr<const FContentBrowserClassFileItemDataPayload> UContentBrowserClassDa
 
 void UContentBrowserClassDataSource::PopulateAddNewContextMenu(UToolMenu* InMenu)
 {
+	if (ensure(GUnrealEd) && !GUnrealEd->GetUnrealEdOptions()->IsCPPAllowed())
+	{
+		return;
+	}
+
 	const UContentBrowserDataMenuContext_AddNewMenu* ContextObject = InMenu->FindContext<UContentBrowserDataMenuContext_AddNewMenu>();
 	checkf(ContextObject, TEXT("Required context UContentBrowserDataMenuContext_AddNewMenu was missing!"));
 
diff --git a/Engine/Source/Developer/AssetTools/Private/AssetTypeActions/AssetTypeActions_Class.cpp b/Engine/Source/Developer/AssetTools/Private/AssetTypeActions/AssetTypeActions_Class.cpp
--- a/Engine/Source/Developer/AssetTools/Private/AssetTypeActions/AssetTypeActions_Class.cpp
+++ b/Engine/Source/Developer/AssetTools/Private/AssetTypeActions/AssetTypeActions_Class.cpp
@@ -9,6 +9,9 @@
 #include "AddToProjectConfig.h"
 #include "GameProjectGenerationModule.h"
 #include "Kismet2/KismetEditorUtilities.h"
+#include "Editor/UnrealEdEngine.h"
+#include "Preferences/UnrealEdOptions.h"
+#include "UnrealEdGlobals.h"
 
 #define LOCTEXT_NAMESPACE "AssetTypeActions"
 
@@ -45,9 +48,14 @@ void FAssetTypeActions_Class::GetActions(const TArray<UObject*>& InObjects, FToo
 		);
 	};
 
-	auto CanCreateDerivedCppClass = [bIsValidBaseCppClass]() -> bool
+	auto IsCPPAllowed = []()
 	{
-		return bIsValidBaseCppClass;
+		return ensure(GUnrealEd) && GUnrealEd->GetUnrealEdOptions()->IsCPPAllowed();
+	};
+
+	auto CanCreateDerivedCppClass = [bIsValidBaseCppClass, IsCPPAllowed]() -> bool
+	{
+		return bIsValidBaseCppClass && IsCPPAllowed();
 	};
 
 	auto CreateCreateDerivedBlueprintClass = [BaseClass]()
@@ -110,9 +118,11 @@ void FAssetTypeActions_Class::GetActions(const TArray<UObject*>& InObjects, FToo
 		FSlateIcon(FEditorStyle::GetStyleSetName(), "Icons.C++"),
 		FUIAction(
 			FExecuteAction::CreateLambda(CreateCreateDerivedCppClass),
-			FCanExecuteAction::CreateLambda(CanCreateDerivedCppClass)
-			)
-		);
+			FCanExecuteAction::CreateLambda(CanCreateDerivedCppClass),
+			FIsActionChecked(),
+			FIsActionButtonVisible::CreateLambda(IsCPPAllowed)
+		)
+	);
 
 	Section.AddMenuEntry(
 		"NewDerivedBlueprintClass",
diff --git a/Engine/Source/Editor/MainFrame/Private/Frame/MainFrameActions.cpp b/Engine/Source/Editor/MainFrame/Private/Frame/MainFrameActions.cpp
--- a/Engine/Source/Editor/MainFrame/Private/Frame/MainFrameActions.cpp
+++ b/Engine/Source/Editor/MainFrame/Private/Frame/MainFrameActions.cpp
@@ -33,6 +33,9 @@
 #include "Settings/EditorExperimentalSettings.h"
 #include "CookerSettings.h"
 #include "UnrealEdMisc.h"
+#include "Editor/UnrealEdEngine.h"
+#include "Preferences/UnrealEdOptions.h"
+#include "UnrealEdGlobals.h"
 #include "FileHelpers.h"
 #include "EditorAnalytics.h"
 #include "LevelEditor.h"
@@ -119,14 +122,14 @@ void FMainFrameCommands::RegisterCommands()
 	UI_COMMAND( OpenProject, "Open Project...", "Opens a dialog to choose a game project to open", EUserInterfaceActionType::Button, FInputChord() );
 	ActionList->MapAction( OpenProject, FExecuteAction::CreateStatic( &FMainFrameActionCallbacks::NewProject, true, false), DefaultExecuteAction );
 
-	UI_COMMAND( AddCodeToProject, "New C++ Class...", "Adds C++ code to the project. The code can only be compiled if you have an appropriate C++ compiler installed.", EUserInterfaceActionType::Button, FInputChord() );
-	ActionList->MapAction( AddCodeToProject, FExecuteAction::CreateStatic( &FMainFrameActionCallbacks::AddCodeToProject ));
+	UI_COMMAND(AddCodeToProject, "New C++ Class...", "Adds C++ code to the project. The code can only be compiled if you have an appropriate C++ compiler installed.", EUserInterfaceActionType::Button, FInputChord());
+	ActionList->MapAction(AddCodeToProject, FExecuteAction::CreateStatic(&FMainFrameActionCallbacks::AddCodeToProject), FCanExecuteAction::CreateStatic(&FMainFrameActionCallbacks::CanAddCodeToProject), FGetActionCheckState(), FIsActionButtonVisible::CreateStatic(&FMainFrameActionCallbacks::IsAddCodeToProjectVisible));
 
-	UI_COMMAND( RefreshCodeProject, "Refresh code project", "Refreshes your C++ code project.", EUserInterfaceActionType::Button, FInputChord() );
-	ActionList->MapAction( RefreshCodeProject, FExecuteAction::CreateStatic( &FMainFrameActionCallbacks::RefreshCodeProject ), FCanExecuteAction::CreateStatic( &FMainFrameActionCallbacks::IsCodeProject ) );
+	UI_COMMAND(RefreshCodeProject, "Refresh code project", "Refreshes your C++ code project.", EUserInterfaceActionType::Button, FInputChord());
+	ActionList->MapAction(RefreshCodeProject, FExecuteAction::CreateStatic(&FMainFrameActionCallbacks::RefreshCodeProject), FCanExecuteAction::CreateStatic(&FMainFrameActionCallbacks::CanRefreshCodeProject), FGetActionCheckState(), FIsActionButtonVisible::CreateStatic(&FMainFrameActionCallbacks::IsRefreshCodeProjectVisible));
 
-	UI_COMMAND( OpenIDE, "Open IDE", "Opens your C++ code in an integrated development environment.", EUserInterfaceActionType::Button, FInputChord() );
-	ActionList->MapAction( OpenIDE, FExecuteAction::CreateStatic( &FMainFrameActionCallbacks::OpenIDE ), FCanExecuteAction::CreateStatic( &FMainFrameActionCallbacks::IsCodeProject ), FGetActionCheckState(), FIsActionButtonVisible::CreateStatic( &FMainFrameActionCallbacks::CanOpenIDE ) );
+	UI_COMMAND(OpenIDE, "Open IDE", "Opens your C++ code in an integrated development environment.", EUserInterfaceActionType::Button, FInputChord());
+	ActionList->MapAction(OpenIDE, FExecuteAction::CreateStatic(&FMainFrameActionCallbacks::OpenIDE), FCanExecuteAction::CreateStatic(&FMainFrameActionCallbacks::CanOpenIDE), FGetActionCheckState(), FIsActionButtonVisible::CreateStatic(&FMainFrameActionCallbacks::IsOpenIDEVisible));
 
 	UI_COMMAND( ZipUpProject, "Zip Project", "Zips the project into a zip file.", EUserInterfaceActionType::Button, FInputChord() );
 	ActionList->MapAction(ZipUpProject, FExecuteAction::CreateStatic( &FMainFrameActionCallbacks::ZipUpProject ), DefaultExecuteAction);
@@ -374,6 +377,16 @@ void FMainFrameActionCallbacks::AddCodeToProject()
 	FGameProjectGenerationModule::Get().OpenAddCodeToProjectDialog();
 }
 
+bool FMainFrameActionCallbacks::CanAddCodeToProject()
+{
+	return IsCPPAllowed();
+}
+
+bool FMainFrameActionCallbacks::IsAddCodeToProjectVisible()
+{
+	return IsCPPAllowed();
+}
+
 void FMainFrameActionCallbacks::RefreshCodeProject()
 {
 	if ( !FSourceCodeNavigation::IsCompilerAvailable() )
@@ -389,6 +402,16 @@ void FMainFrameActionCallbacks::RefreshCodeProject()
 	}
 }
 
+bool FMainFrameActionCallbacks::CanRefreshCodeProject()
+{
+	return IsCPPAllowed() && IsCodeProject();
+}
+
+bool FMainFrameActionCallbacks::IsRefreshCodeProjectVisible()
+{
+	return IsCPPAllowed();
+}
+
 bool FMainFrameActionCallbacks::IsCodeProject()
 {
 	// Not particularly rigorous, but assume it's a code project if it can find a Source directory
@@ -414,7 +437,17 @@ void FMainFrameActionCallbacks::OpenIDE()
 
 bool FMainFrameActionCallbacks::CanOpenIDE()
 {
-	return FSourceCodeNavigation::DoesModuleSolutionExist();
+	return IsCPPAllowed() && IsCodeProject();
+}
+
+bool FMainFrameActionCallbacks::IsOpenIDEVisible()
+{
+	return IsCPPAllowed() && FSourceCodeNavigation::DoesModuleSolutionExist();
+}
+
+bool FMainFrameActionCallbacks::IsCPPAllowed()
+{
+	return ensure(GUnrealEd) && GUnrealEd->GetUnrealEdOptions()->IsCPPAllowed();
 }
 
 void FMainFrameActionCallbacks::ZipUpProject()
diff --git a/Engine/Source/Editor/MainFrame/Private/Frame/MainFrameActions.h b/Engine/Source/Editor/MainFrame/Private/Frame/MainFrameActions.h
--- a/Engine/Source/Editor/MainFrame/Private/Frame/MainFrameActions.h
+++ b/Engine/Source/Editor/MainFrame/Private/Frame/MainFrameActions.h
@@ -184,21 +184,39 @@ public:
 	/** Adds code to the current project if it does not already have any */
 	static void AddCodeToProject();
 
+	/** Whether we can add code to the current project */
+	static bool CanAddCodeToProject();
+
+	/** Whether menu to add code to the current project is visible */
+	static bool IsAddCodeToProjectVisible();
+
 	/** Checks whether a menu action for packaging the project can execute. */
 	static bool PackageProjectCanExecute( const FName PlatformInfoName );
 
 	/** Refresh the project in the current IDE */
 	static void RefreshCodeProject();
 
+	/** Whether we can refresh the project in the current IDE */
+	static bool CanRefreshCodeProject();
+
+	/** Whether menu to refresh the project in the current IDE is visible */
+	static bool IsRefreshCodeProjectVisible();
+
 	/** Determines whether the project is a code project */
 	static bool IsCodeProject();
 
 	/** Opens an IDE to edit c++ code */
 	static void OpenIDE();
 
-	/** Determines whether we can open the IDE to edit c++ code */
+	/** Whether we can open the IDE to edit c++ code */
 	static bool CanOpenIDE();
 
+	/** Whether menu to open the IDE to edit c++ code is visible */
+	static bool IsOpenIDEVisible();
+
+	/** Whether C++ edition/creation is allowed from the editor */
+	static bool IsCPPAllowed();
+
 	/** Zips up the project */
 	static void ZipUpProject();
 
diff --git a/Engine/Source/Editor/UMGEditor/Private/WidgetBlueprintEditor.cpp b/Engine/Source/Editor/UMGEditor/Private/WidgetBlueprintEditor.cpp
--- a/Engine/Source/Editor/UMGEditor/Private/WidgetBlueprintEditor.cpp
+++ b/Engine/Source/Editor/UMGEditor/Private/WidgetBlueprintEditor.cpp
@@ -74,6 +74,10 @@
 #include "TabFactory/AnimationTabSummoner.h"
 #include "Kismet/Public/BlueprintEditorTabs.h"
 
+#include "Editor/UnrealEdEngine.h"
+#include "Preferences/UnrealEdOptions.h"
+#include "UnrealEdGlobals.h"
+
 #define LOCTEXT_NAMESPACE "UMG"
 
 FWidgetBlueprintEditor::FWidgetBlueprintEditor()
@@ -281,7 +285,9 @@ void FWidgetBlueprintEditor::BindToolkitCommands()
 	GetToolkitCommands()->MapAction(FUMGEditorCommands::Get().CreateNativeBaseClass,
 		FUIAction(
 			FExecuteAction::CreateSP(this, &FWidgetBlueprintEditor::OpenCreateNativeBaseClassDialog),
-			FCanExecuteAction::CreateSP(this, &FWidgetBlueprintEditor::IsParentClassNative)
+			FCanExecuteAction::CreateSP(this, &FWidgetBlueprintEditor::CanCreateNativeBaseClass),
+			FGetActionCheckState(),
+			FIsActionButtonVisible::CreateSP(this, &FWidgetBlueprintEditor::IsCreateNativeBaseClassVisible)
 		)
 	);
 
@@ -846,6 +852,17 @@ void FWidgetBlueprintEditor::DuplicateSelectedWidgets()
 	SelectWidgets(DuplicatedWidgetRefs, false);
 }
 
+bool FWidgetBlueprintEditor::CanCreateNativeBaseClass() const
+{
+	return ensure(GUnrealEd) && GUnrealEd->GetUnrealEdOptions()->IsCPPAllowed() && IsParentClassNative();
+}
+
+
+bool FWidgetBlueprintEditor::IsCreateNativeBaseClassVisible() const
+{
+	return ensure(GUnrealEd) && GUnrealEd->GetUnrealEdOptions()->IsCPPAllowed();
+}
+
 void FWidgetBlueprintEditor::Tick(float DeltaTime)
 {
 	Super::Tick(DeltaTime);
diff --git a/Engine/Source/Editor/UMGEditor/Public/WidgetBlueprintEditor.h b/Engine/Source/Editor/UMGEditor/Public/WidgetBlueprintEditor.h
--- a/Engine/Source/Editor/UMGEditor/Public/WidgetBlueprintEditor.h
+++ b/Engine/Source/Editor/UMGEditor/Public/WidgetBlueprintEditor.h
@@ -300,6 +300,12 @@ private:
 	bool CanDuplicateSelectedWidgets();
 	void DuplicateSelectedWidgets();
 
+	/** Is creating a native base class for the current widget blueprint allowed */
+	bool CanCreateNativeBaseClass() const;
+
+	/** Whether menu to create a native base class is visible */
+	bool IsCreateNativeBaseClassVisible() const;
+
 private:
 	/** Called whenever the blueprint is structurally changed. */
 	virtual void OnBlueprintChangedImpl(UBlueprint* InBlueprint, bool bIsJustBeingCompiled = false ) override;
