Removing early out based on bShowProjectMenus and changing it to a condition for adding menu entries, because other sections rely on FileProject

diff --git a/Engine/Config/BaseEditorPerProjectUserSettings.ini b/Engine/Config/BaseEditorPerProjectUserSettings.ini
--- a/Engine/Config/BaseEditorPerProjectUserSettings.ini
+++ b/Engine/Config/BaseEditorPerProjectUserSettings.ini
@@ -75,7 +75,7 @@ Collision=True
 VertexColor=true
 MapSkeletalMotionToRoot=False
 
-[/Script/EditorStyle.EditorStyleSettings]
+[/Script/UnrealEd.EditorStyleSettings]
 ; Whether to enable the Editor UI Layout configuration tools for the user
 bEnableUserEditorLayoutManagement=True
 
@@ -97,6 +97,7 @@ bUseSmallToolBarIcons=False
 bShowFriendlyNames=True
 ;If true, the Editor Preferences and Project Settings menu entries in the main menu will be expanded with sub-menus for each settings section
 bExpandConfigurationMenus=False
+;If true, the Project section of the File menu will be shown
 bShowProjectMenus=True
 bShowLaunchMenus=True
 bShowAllAdvancedDetails=False
diff --git a/Engine/Source/Editor/MainFrame/Private/Menus/MainMenu.cpp b/Engine/Source/Editor/MainFrame/Private/Menus/MainMenu.cpp
--- a/Engine/Source/Editor/MainFrame/Private/Menus/MainMenu.cpp
+++ b/Engine/Source/Editor/MainFrame/Private/Menus/MainMenu.cpp
@@ -396,47 +396,44 @@ void FMainMenu::RegisterMainMenu()
 
 void FMainMenu::RegisterFileProjectMenu()
 {
-	if (!GetDefault<UEditorStyleSettings>()->bShowProjectMenus)
-	{
-		return;
-	}
-
 	UToolMenus* ToolMenus = UToolMenus::Get();
 	UToolMenu* MainTabFileMenu = ToolMenus->ExtendMenu("MainFrame.MainTabMenu.File");
 	FToolMenuSection& Section = MainTabFileMenu->AddSection("FileProject", LOCTEXT("ProjectHeading", "Project"));
 
-	Section.AddMenuEntry( FMainFrameCommands::Get().NewProject );
-	Section.AddMenuEntry( FMainFrameCommands::Get().OpenProject );
-	/*
-	MenuBuilder.AddMenuEntry( FMainFrameCommands::Get().LocalizeProject,
-		NAME_None,
-		TAttribute<FText>(),
-		LOCTEXT("LocalizeProjectToolTip", "Gather text from your project and import/export translations.")
-		);
-		*/
-	/*
-	MenuBuilder.AddSubMenu(
-		LOCTEXT("CookProjectSubMenuLabel", "Cook Project"),
-		LOCTEXT("CookProjectSubMenuToolTip", "Cook your project content for debugging"),
-		FNewMenuDelegate::CreateStatic( &FCookContentMenu::MakeMenu ), false, FSlateIcon()
-	);
-	*/
+	if (GetDefault<UEditorStyleSettings>()->bShowProjectMenus)
+	{
+		Section.AddMenuEntry(FMainFrameCommands::Get().NewProject);
+		Section.AddMenuEntry(FMainFrameCommands::Get().OpenProject);
+		/*
+		MenuBuilder.AddMenuEntry( FMainFrameCommands::Get().LocalizeProject,
+			NAME_None,
+			TAttribute<FText>(),
+			LOCTEXT("LocalizeProjectToolTip", "Gather text from your project and import/export translations.")
+			);
+			*/
+			/*
+			MenuBuilder.AddSubMenu(
+				LOCTEXT("CookProjectSubMenuLabel", "Cook Project"),
+				LOCTEXT("CookProjectSubMenuToolTip", "Cook your project content for debugging"),
+				FNewMenuDelegate::CreateStatic( &FCookContentMenu::MakeMenu ), false, FSlateIcon()
+			);
+			*/
 
 
-	Section.AddMenuEntry(FMainFrameCommands::Get().ZipUpProject);
+		Section.AddMenuEntry(FMainFrameCommands::Get().ZipUpProject);
 
-	if (GetDefault<UEditorStyleSettings>()->bShowProjectMenus && FMainFrameActionCallbacks::RecentProjects.Num() > 0)
-	{
-		Section.AddSubMenu(
-			"RecentProjects",
-			LOCTEXT("SwitchProjectSubMenu", "Recent Projects"),
-			LOCTEXT("SwitchProjectSubMenu_ToolTip", "Select a project to switch to"),
-			FNewToolMenuDelegate::CreateStatic(&FRecentProjectsMenu::MakeMenu),
-			false,
-			FSlateIcon(FAppStyle::GetAppStyleSetName(), "MainFrame.RecentProjects")
-		);
+		if (FMainFrameActionCallbacks::RecentProjects.Num() > 0)
+		{
+			Section.AddSubMenu(
+				"RecentProjects",
+				LOCTEXT("SwitchProjectSubMenu", "Recent Projects"),
+				LOCTEXT("SwitchProjectSubMenu_ToolTip", "Select a project to switch to"),
+				FNewToolMenuDelegate::CreateStatic(&FRecentProjectsMenu::MakeMenu),
+				false,
+				FSlateIcon(FAppStyle::GetAppStyleSetName(), "MainFrame.RecentProjects")
+			);
+		}
 	}
-
 }
 
 void FMainMenu::RegisterToolsMenu()
diff --git a/Engine/Source/Editor/UnrealEd/Classes/Settings/EditorStyleSettings.h b/Engine/Source/Editor/UnrealEd/Classes/Settings/EditorStyleSettings.h
--- a/Engine/Source/Editor/UnrealEd/Classes/Settings/EditorStyleSettings.h
+++ b/Engine/Source/Editor/UnrealEd/Classes/Settings/EditorStyleSettings.h
@@ -137,7 +137,7 @@ public:
 	UPROPERTY(EditAnywhere, config, Category=UserInterface, AdvancedDisplay)
 	uint32 bExpandConfigurationMenus:1;
 
-	/** When enabled, the Editor Preferences and Project Settings menu items in the main menu will be expanded with sub-menus for each settings section. */
+	/** When enabled, the project subsection of the File menu will be shown. */
 	UPROPERTY(config)
 	uint32 bShowProjectMenus : 1;
 
