Fixed behavior when selecting "Change Source Control Settings..." from the Editor top menu to preserve the actual setting on cancel to be consisent with the behavior when doing the same thing from the bottom right button.

diff --git a/Engine/Source/Editor/MainFrame/Private/Frame/MainFrameActions.cpp b/Engine/Source/Editor/MainFrame/Private/Frame/MainFrameActions.cpp
--- a/Engine/Source/Editor/MainFrame/Private/Frame/MainFrameActions.cpp
+++ b/Engine/Source/Editor/MainFrame/Private/Frame/MainFrameActions.cpp
@@ -299,7 +299,7 @@ bool FMainFrameActionCallbacks::CanViewChangelists()
 void FMainFrameActionCallbacks::ConnectToSourceControl()
 {
 	ELoginWindowMode::Type Mode = !FSlateApplication::Get().GetActiveModalWindow().IsValid() ? ELoginWindowMode::Modeless : ELoginWindowMode::Modal;
-	ISourceControlModule::Get().ShowLoginDialog(FSourceControlLoginClosed(), Mode);
+	ISourceControlModule::Get().ShowLoginDialog(FSourceControlLoginClosed(), Mode, EOnLoginWindowStartup::PreserveProvider);
 }
 
 bool FMainFrameActionCallbacks::CanSaveWorld()
