Finalize links to community pages in the Help menu

diff --git a/Engine/Config/BaseEditor.ini b/Engine/Config/BaseEditor.ini
--- a/Engine/Config/BaseEditor.ini
+++ b/Engine/Config/BaseEditor.ini
@@ -269,10 +269,12 @@ bShouldPromptOnCurveRemoveAll=true
 
 [UnrealEd.URLs]
 DocumentationURL="https://docs.unrealengine.com/{VERSION}/{I18N}/"
-OnlineLearningURL="https://www.unrealengine.com/onlinelearning-courses"
-ForumsURL="http://forums.unrealengine.com"
 APIDocsURL="https://docs.unrealengine.com/{VERSION}/{I18N}/API/"
+CommunityHomeURL="https://dev.epicgames.com/community"
+OnlineLearningURL="https://dev.epicgames.com/community/learning"
+ForumsURL="http://forums.unrealengine.com"
 SearchForAnswersURL="https://forums.unrealengine.com/tag/question"
+SnippetsURL="https://dev.epicgames.com/community/snippets"
 SupportWebsiteURL="http://unrealengine.com/support"
 SourceCodeIDEURL_Windows="https://aka.ms/vs/16/release/vs_Community.exe"
 SourceCodeIDEURL_Mac="https://developer.apple.com/xcode"
diff --git a/Engine/Source/Editor/EditorStyle/Private/StarshipStyle.cpp b/Engine/Source/Editor/EditorStyle/Private/StarshipStyle.cpp
--- a/Engine/Source/Editor/EditorStyle/Private/StarshipStyle.cpp
+++ b/Engine/Source/Editor/EditorStyle/Private/StarshipStyle.cpp
@@ -3990,13 +3990,15 @@ void FStarshipEditorStyle::FStyle::SetupLevelEditorStyle()
 		Set("MainFrame.DocumentationHome",	        new IMAGE_BRUSH_SVG("Starship/Common/Documentation", Icon16x16));
 		Set("MainFrame.BrowseAPIReference",         new IMAGE_BRUSH_SVG("Starship/Common/Documentation", Icon16x16));
 		Set("MainFrame.BrowseCVars",                new CORE_IMAGE_BRUSH_SVG("Starship/Common/Console", Icon16x16));
+		Set("MainFrame.VisitCommunityHome",         new IMAGE_BRUSH_SVG("Starship/Common/Community", Icon16x16));
 		Set("MainFrame.VisitOnlineLearning",		new IMAGE_BRUSH_SVG("Starship/Common/Tutorials", Icon16x16));
 		Set("MainFrame.VisitForums",                new IMAGE_BRUSH_SVG("Starship/Common/Forums", Icon16x16));
 		Set("MainFrame.VisitSearchForAnswersPage",  new IMAGE_BRUSH_SVG("Starship/Common/QuestionAnswer", Icon16x16));
-		Set("MainFrame.VisitSupportWebSite",       new IMAGE_BRUSH_SVG("Starship/Common/Support", Icon16x16));
+		Set("MainFrame.VisitCommunitySnippets",     new IMAGE_BRUSH_SVG("Starship/Common/FileLined", Icon16x16));
+		Set("MainFrame.VisitSupportWebSite",        new IMAGE_BRUSH_SVG("Starship/Common/Support", Icon16x16));
 		Set("MainFrame.VisitEpicGamesDotCom",       new IMAGE_BRUSH_SVG("About/EpicGamesLogo", Icon16x16));
 		Set("MainFrame.AboutUnrealEd",              new IMAGE_BRUSH_SVG("About/UnrealLogo", Icon16x16));
-		Set("MainFrame.CreditsUnrealEd",           new IMAGE_BRUSH_SVG("Starship/Common/Credits", Icon16x16));
+		Set("MainFrame.CreditsUnrealEd",            new IMAGE_BRUSH_SVG("Starship/Common/Credits", Icon16x16));
 
 		Set( "EditorViewport.SelectMode", new IMAGE_BRUSH_SVG("Starship/EditorViewport/select", Icon16x16) );
 		Set( "EditorViewport.TranslateMode", new IMAGE_BRUSH_SVG( "Starship/EditorViewport/translate", Icon16x16 ) );
diff --git a/Engine/Source/Editor/MainFrame/Private/Frame/MainFrameActions.cpp b/Engine/Source/Editor/MainFrame/Private/Frame/MainFrameActions.cpp
--- a/Engine/Source/Editor/MainFrame/Private/Frame/MainFrameActions.cpp
+++ b/Engine/Source/Editor/MainFrame/Private/Frame/MainFrameActions.cpp
@@ -186,21 +186,27 @@ void FMainFrameCommands::RegisterCommands()
 	UI_COMMAND(DocumentationHome, "Documentation Home", "Authoritative, in-depth technical resources for using Unreal Engine", EUserInterfaceActionType::Button, FInputChord());
 	ActionList->MapAction(DocumentationHome, FExecuteAction::CreateStatic(&FMainFrameActionCallbacks::DocumentationHome));
 
-	UI_COMMAND(VisitOnlineLearning, "Online Learning", "Learn Unreal Engine for free with easy-to-follow video courses and guided learning paths", EUserInterfaceActionType::Button, FInputChord());
-	ActionList->MapAction(VisitOnlineLearning, FExecuteAction::CreateStatic(&FMainFrameActionCallbacks::VisitOnlineLearning));
-
 	UI_COMMAND(BrowseAPIReference, "C++ API Reference", "Classes, functions, and other elements that make up the C++ API", EUserInterfaceActionType::Button, FInputChord());
 	ActionList->MapAction(BrowseAPIReference, FExecuteAction::CreateStatic(&FMainFrameActionCallbacks::BrowseAPIReference));
 
 	UI_COMMAND(BrowseCVars, "Console Variables", "Reference companion for console variables and commands", EUserInterfaceActionType::Button, FInputChord());
 	ActionList->MapAction(BrowseCVars, FExecuteAction::CreateStatic(&FMainFrameActionCallbacks::BrowseCVars));
 
+	UI_COMMAND(VisitCommunityHome, "Dev Community", "Join the worldwide community of Unreal developers and explore all the resources it has to offer", EUserInterfaceActionType::Button, FInputChord());
+	ActionList->MapAction(VisitCommunityHome, FExecuteAction::CreateStatic(&FMainFrameActionCallbacks::VisitCommunityHome));
+
+	UI_COMMAND(VisitOnlineLearning, "Learning Library", "Learn Unreal Engine for free with easy-to-follow video courses and guided learning paths", EUserInterfaceActionType::Button, FInputChord());
+	ActionList->MapAction(VisitOnlineLearning, FExecuteAction::CreateStatic(&FMainFrameActionCallbacks::VisitOnlineLearning));
+
 	UI_COMMAND(VisitForums, "Forums", "View announcements and engage in discussions with other developers", EUserInterfaceActionType::Button, FInputChord());
 	ActionList->MapAction(VisitForums, FExecuteAction::CreateStatic(&FMainFrameActionCallbacks::VisitForums));
 
 	UI_COMMAND(VisitSearchForAnswersPage, "Q&A", "Search for answers, ask questions, and share your knowledge with other developers", EUserInterfaceActionType::Button, FInputChord());
 	ActionList->MapAction(VisitSearchForAnswersPage, FExecuteAction::CreateStatic(&FMainFrameActionCallbacks::VisitSearchForAnswersPage));
 
+	UI_COMMAND(VisitCommunitySnippets, "Snippets", "Access and share ready-to-use code blocks and scripts", EUserInterfaceActionType::Button, FInputChord());
+	ActionList->MapAction(VisitCommunitySnippets, FExecuteAction::CreateStatic(&FMainFrameActionCallbacks::VisitCommunitySnippets));
+
 	UI_COMMAND(VisitSupportWebSite, "Support", "Options for personalized technical support", EUserInterfaceActionType::Button, FInputChord());
 	ActionList->MapAction(VisitSupportWebSite, FExecuteAction::CreateStatic(&FMainFrameActionCallbacks::VisitSupportWebSite));
 
@@ -666,6 +672,15 @@ void FMainFrameActionCallbacks::BrowseCVars()
 	GEditor->Exec(GEditor->GetEditorWorldContext().World(), TEXT("help"));
 }
 
+void FMainFrameActionCallbacks::VisitCommunityHome()
+{
+	FString URL;
+	if (FUnrealEdMisc::Get().GetURL(TEXT("CommunityHomeURL"), URL))
+	{
+		FPlatformProcess::LaunchURL(*URL, NULL, NULL);
+	}
+}
+
 void FMainFrameActionCallbacks::VisitForums()
 {
 	FString URL;
@@ -675,6 +690,15 @@ void FMainFrameActionCallbacks::VisitForums()
 	}
 }
 
+void FMainFrameActionCallbacks::VisitCommunitySnippets()
+{
+	FString URL;
+	if (FUnrealEdMisc::Get().GetURL(TEXT("SnippetsURL"), URL))
+	{
+		FPlatformProcess::LaunchURL(*URL, NULL, NULL);
+	}
+}
+
 void FMainFrameActionCallbacks::AboutUnrealEd_Execute()
 {
 	const FText AboutWindowTitle = LOCTEXT( "AboutUnrealEditor", "About Unreal Editor" );
diff --git a/Engine/Source/Editor/MainFrame/Private/Frame/MainFrameActions.h b/Engine/Source/Editor/MainFrame/Private/Frame/MainFrameActions.h
--- a/Engine/Source/Editor/MainFrame/Private/Frame/MainFrameActions.h
+++ b/Engine/Source/Editor/MainFrame/Private/Frame/MainFrameActions.h
@@ -53,11 +53,13 @@ public:
 	TSharedPtr< FUICommandInfo > OpenMarketplace;
 
 	TSharedPtr< FUICommandInfo > DocumentationHome;
-	TSharedPtr< FUICommandInfo > VisitOnlineLearning;
 	TSharedPtr< FUICommandInfo > BrowseAPIReference;
 	TSharedPtr< FUICommandInfo > BrowseCVars;
+	TSharedPtr< FUICommandInfo > VisitCommunityHome;
+	TSharedPtr< FUICommandInfo > VisitOnlineLearning;
 	TSharedPtr< FUICommandInfo > VisitForums;
 	TSharedPtr< FUICommandInfo > VisitSearchForAnswersPage;
+	TSharedPtr< FUICommandInfo > VisitCommunitySnippets;
 	TSharedPtr< FUICommandInfo > ReportABug;
 	TSharedPtr< FUICommandInfo > OpenIssueTracker;
 	TSharedPtr< FUICommandInfo > VisitSupportWebSite;
@@ -149,9 +151,6 @@ public:
 	/** Opens the issue tracker page */
 	static void OpenIssueTracker();
 
-	/** Visits the "search for answers" page on UDN */
-	static void VisitSearchForAnswersPage();
-
 	/** Visits the UDN support web site */
 	static void VisitSupportWebSite();
 
@@ -161,17 +160,27 @@ public:
 	/** Opens the documentation home page*/
 	static void DocumentationHome();
 
-	/** Visits The Unreal Online Learning page*/
-	static void VisitOnlineLearning();
-
 	/** Opens the API documentation site */
 	static void BrowseAPIReference();
 
 	/** Creates an HTML file to browse the console variables and commands */
 	static void BrowseCVars();
 
+	/** Visits the home page of the Epic Games Dev Community web site */
+	static void VisitCommunityHome();
+
+	/** Visits the Learning Library on the Epic Games Dev Community web site */
+	static void VisitOnlineLearning();
+
+	/** Visits the Unreal Engine community forums */
 	static void VisitForums();
 
+	/** Visits the Q&A section of the Unreal Engine community forums */
+	static void VisitSearchForAnswersPage();
+
+	/** Visits the Snippets Repository on the Epic Games Dev Community web site */
+	static void VisitCommunitySnippets();
+
 	static void AboutUnrealEd_Execute();
 
 	static void CreditsUnrealEd_Execute();
diff --git a/Engine/Source/Editor/MainFrame/Private/Menus/MainMenu.cpp b/Engine/Source/Editor/MainFrame/Private/Menus/MainMenu.cpp
--- a/Engine/Source/Editor/MainFrame/Private/Menus/MainMenu.cpp
+++ b/Engine/Source/Editor/MainFrame/Private/Menus/MainMenu.cpp
@@ -270,9 +270,11 @@ void FMainMenu::RegisterHelpMenu()
 
 	FToolMenuSection& CommunitySection = Menu->AddSection("Community", NSLOCTEXT("MainHelpMenu", "CommunitySection", "Community"));
 	{
+		CommunitySection.AddMenuEntry(FMainFrameCommands::Get().VisitCommunityHome);
 		CommunitySection.AddMenuEntry(FMainFrameCommands::Get().VisitOnlineLearning);
 		CommunitySection.AddMenuEntry(FMainFrameCommands::Get().VisitForums);
 		CommunitySection.AddMenuEntry(FMainFrameCommands::Get().VisitSearchForAnswersPage);
+		CommunitySection.AddMenuEntry(FMainFrameCommands::Get().VisitCommunitySnippets);
 	}
 
 	FToolMenuSection& BugReportingSection = Menu->AddSection("Support", NSLOCTEXT("MainHelpMenu", "SupportSection", "Support"));
