Fix to Zip Project command to focus back to the UE after the dialog is cancelled.

diff --git a/Engine/Source/Editor/MainFrame/Private/Frame/MainFrameActions.cpp b/Engine/Source/Editor/MainFrame/Private/Frame/MainFrameActions.cpp
--- a/Engine/Source/Editor/MainFrame/Private/Frame/MainFrameActions.cpp
+++ b/Engine/Source/Editor/MainFrame/Private/Frame/MainFrameActions.cpp
@@ -464,7 +464,7 @@ void FMainFrameActionCallbacks::ZipUpProject()
 	if (DesktopPlatform != NULL)
 	{
 		bOpened = DesktopPlatform->SaveFileDialog(
-			NULL,
+			FSlateApplication::Get().FindBestParentWindowHandleForDialogs(nullptr),
 			NSLOCTEXT("UnrealEd", "ZipUpProject", "Zip file location").ToString(),
 			FPaths::ProjectDir(),
 			FApp::GetProjectName(),
