Add a way to override the section that gets selected by default when opening the editor preferences dialog

diff --git a/Engine/Source/Editor/MainFrame/Private/MainFrameModule.h b/Engine/Source/Editor/MainFrame/Private/MainFrameModule.h
--- a/Engine/Source/Editor/MainFrame/Private/MainFrameModule.h
+++ b/Engine/Source/Editor/MainFrame/Private/MainFrameModule.h
@@ -171,6 +171,18 @@ public:
 		return true;
 	}
 
+	virtual void SetEditorSettingsDefaultSelectionOverride(FName CategoryName = FName(), FName SectionName = FName()) override
+	{
+		EditorSettingsDefaultCategoryOverride = CategoryName;
+		EditorSettingsDefaultSectionOverride = SectionName;
+	}
+
+	virtual void GetEditorSettingsDefaultSelectionOverride(FName& OutCategoryName, FName& OutSectionName) override
+	{
+		OutCategoryName = EditorSettingsDefaultCategoryOverride;
+		OutSectionName = EditorSettingsDefaultSectionOverride;
+	}
+
 public:
 
 	// IModuleInterface interface
@@ -239,6 +251,12 @@ private:
 	// Override window title, or empty to not override
 	FText OverriddenWindowTitle;
 
+	// Overrides the category that gets selected by default when opening editor settings
+	FName EditorSettingsDefaultCategoryOverride;
+
+	// Overrides the section that gets selected by default when editor settings
+	FName EditorSettingsDefaultSectionOverride;
+
 	/// Event to be called when the mainframe is fully created.
 	FMainFrameCreationFinishedEvent MainFrameCreationFinishedEvent;
 
diff --git a/Engine/Source/Editor/MainFrame/Private/Menus/MainMenu.cpp b/Engine/Source/Editor/MainFrame/Private/Menus/MainMenu.cpp
--- a/Engine/Source/Editor/MainFrame/Private/Menus/MainMenu.cpp
+++ b/Engine/Source/Editor/MainFrame/Private/Menus/MainMenu.cpp
@@ -130,14 +130,30 @@ void FMainMenu::RegisterEditMenu()
 		else
 		{
 #if !PLATFORM_MAC // Handled by app's menu in menu bar
-			Section.AddMenuEntry(
-				"EditorPreferencesMenu",
-				LOCTEXT("EditorPreferencesMenuLabel", "Editor Preferences..."),
-				LOCTEXT("EditorPreferencesMenuToolTip", "Configure the behavior and features of the Unreal Editor."),
-				FSlateIcon(FAppStyle::GetAppStyleSetName(), "EditorPreferences.TabIcon"),
-				FUIAction(FExecuteAction::CreateStatic(&FSettingsMenu::OpenSettings, FName("Editor"), FName("General"), FName("Appearance")))
-			);
-#endif
+			{
+				IMainFrameModule& MainFrame = FModuleManager::LoadModuleChecked<IMainFrameModule>("MainFrame");
+				FName CategoryName;
+				FName SectionName;
+				MainFrame.GetEditorSettingsDefaultSelectionOverride(CategoryName, SectionName);
+
+				if (CategoryName.IsNone())
+				{
+					CategoryName = FName("General");
+				}
+				if (SectionName.IsNone())
+				{
+					SectionName = FName("Appearance");
+				}
+
+				Section.AddMenuEntry(
+					"EditorPreferencesMenu",
+					LOCTEXT("EditorPreferencesMenuLabel", "Editor Preferences..."),
+					LOCTEXT("EditorPreferencesMenuToolTip", "Configure the behavior and features of the Unreal Editor."),
+					FSlateIcon(FAppStyle::GetAppStyleSetName(), "EditorPreferences.TabIcon"),
+					FUIAction(FExecuteAction::CreateStatic(&FSettingsMenu::OpenSettings, FName("Editor"), CategoryName, SectionName))
+				);
+			}
+#endif //if !PLATFORM_MAC
 
 			Section.AddMenuEntry(
 				"ProjectSettingsMenu",
diff --git a/Engine/Source/Editor/MainFrame/Public/Interfaces/IMainFrameModule.h b/Engine/Source/Editor/MainFrame/Public/Interfaces/IMainFrameModule.h
--- a/Engine/Source/Editor/MainFrame/Public/Interfaces/IMainFrameModule.h
+++ b/Engine/Source/Editor/MainFrame/Public/Interfaces/IMainFrameModule.h
@@ -216,6 +216,18 @@ public:
 
 	virtual bool ExecuteCanCloseEditorDelegates() = 0;
 
+	/**
+	 * Overrides the section that gets selected by default when opening editor settings
+	 * @note Call with default parameters to clear the override
+	 */
+	virtual void SetEditorSettingsDefaultSelectionOverride(FName CategoryName = FName(), FName SectionName = FName()) = 0;
+
+	/** 
+	 * Gets the override for the section that gets selected by default when opening editor settings
+	 * @note Returns None if there's no override
+	 */
+	virtual void GetEditorSettingsDefaultSelectionOverride(FName& OutCategoryName, FName& OutSectionName) = 0;
+
 public:
 
 	/**
