Providing a hook for external sources to override the editor's main window, much like the project browser does. Allows plugins to provide their own kind of "start up check" dialog.

Switching over the project browser dialog to use the new system.

diff --git a/Engine/Plugins/Editor/GuidedTutorials/Source/IntroTutorials/Private/IntroTutorials.cpp b/Engine/Plugins/Editor/GuidedTutorials/Source/IntroTutorials/Private/IntroTutorials.cpp
--- a/Engine/Plugins/Editor/GuidedTutorials/Source/IntroTutorials/Private/IntroTutorials.cpp
+++ b/Engine/Plugins/Editor/GuidedTutorials/Source/IntroTutorials/Private/IntroTutorials.cpp
@@ -268,9 +268,9 @@ void FIntroTutorials::AddSummonBlueprintTutorialsMenuExtension(FMenuBuilder& Men
 	MenuBuilder.EndSection();
 }
 
-void FIntroTutorials::MainFrameLoad(TSharedPtr<SWindow> InRootWindow, bool bIsNewProjectWindow)
+void FIntroTutorials::MainFrameLoad(TSharedPtr<SWindow> InRootWindow, bool bIsRunningStartupDialog)
 {
-	if (!bIsNewProjectWindow)
+	if (!bIsRunningStartupDialog)
 	{
 		// install a root widget for the tutorial overlays to hang off
 		if(InRootWindow.IsValid() && !TutorialRoot.IsValid())
diff --git a/Engine/Plugins/Editor/GuidedTutorials/Source/IntroTutorials/Private/IntroTutorials.h b/Engine/Plugins/Editor/GuidedTutorials/Source/IntroTutorials/Private/IntroTutorials.h
--- a/Engine/Plugins/Editor/GuidedTutorials/Source/IntroTutorials/Private/IntroTutorials.h
+++ b/Engine/Plugins/Editor/GuidedTutorials/Source/IntroTutorials/Private/IntroTutorials.h
@@ -73,7 +73,7 @@ private:
 	void AddSummonBlueprintTutorialsMenuExtension(FMenuBuilder& MenuBuilder, UObject* PrimaryObject);
 
 	/** Event to be called when the main frame is loaded */
-	void MainFrameLoad(TSharedPtr<SWindow> InRootWindow, bool bIsNewProjectWindow);
+	void MainFrameLoad(TSharedPtr<SWindow> InRootWindow, bool bIsRunningStartupDialog);
 
 	/** Summon tutorial home page to front */
 	void SummonTutorialHome();
diff --git a/Engine/Plugins/Editor/MacGraphicsSwitching/Source/MacGraphicsSwitching/Private/MacGraphicsSwitchingModule.cpp b/Engine/Plugins/Editor/MacGraphicsSwitching/Source/MacGraphicsSwitching/Private/MacGraphicsSwitchingModule.cpp
--- a/Engine/Plugins/Editor/MacGraphicsSwitching/Source/MacGraphicsSwitching/Private/MacGraphicsSwitchingModule.cpp
+++ b/Engine/Plugins/Editor/MacGraphicsSwitching/Source/MacGraphicsSwitching/Private/MacGraphicsSwitchingModule.cpp
@@ -26,7 +26,7 @@ public:
 	virtual void ShutdownModule() override;
 	
 private:
-	void Initialize( TSharedPtr<SWindow> InRootWindow, bool bIsNewProjectWindow );
+	void Initialize( TSharedPtr<SWindow> InRootWindow, bool bIsRunningStartupDialog );
 	void AddGraphicsSwitcher( FToolBarBuilder& ToolBarBuilder );
 	
 private:
@@ -56,9 +56,9 @@ void FMacGraphicsSwitchingModule::StartupModule()
 	}
 }
 
-void FMacGraphicsSwitchingModule::Initialize( TSharedPtr<SWindow> InRootWindow, bool bIsNewProjectWindow )
+void FMacGraphicsSwitchingModule::Initialize( TSharedPtr<SWindow> InRootWindow, bool bIsRunningStartupDialog )
 {
-	if( !bIsNewProjectWindow )
+	if( !bIsRunningStartupDialog )
 	{
 		IMainFrameModule& MainFrameModule = FModuleManager::LoadModuleChecked<IMainFrameModule>(TEXT("MainFrame"));
 		MainFrameModule.OnMainFrameCreationFinished().RemoveAll(this);
diff --git a/Engine/Plugins/Editor/PluginBrowser/Source/PluginBrowser/Private/PluginBrowserModule.cpp b/Engine/Plugins/Editor/PluginBrowser/Source/PluginBrowser/Private/PluginBrowserModule.cpp
--- a/Engine/Plugins/Editor/PluginBrowser/Source/PluginBrowser/Private/PluginBrowserModule.cpp
+++ b/Engine/Plugins/Editor/PluginBrowser/Source/PluginBrowser/Private/PluginBrowserModule.cpp
@@ -194,10 +194,10 @@ TSharedRef<SDockTab> FPluginBrowserModule::SpawnPluginCreatorTab(const FSpawnTab
 	return ResultTab;
 }
 
-void FPluginBrowserModule::OnMainFrameLoaded(TSharedPtr<SWindow> InRootWindow, bool bIsNewProjectWindow)
+void FPluginBrowserModule::OnMainFrameLoaded(TSharedPtr<SWindow> InRootWindow, bool bIsRunningStartupDialog)
 {
 	// Show a popup notification that allows the user to enable any new plugins
-	if(!bIsNewProjectWindow && NewlyInstalledPlugins.Num() > 0 && !PluginBrowserTab.IsValid())
+	if(!bIsRunningStartupDialog && NewlyInstalledPlugins.Num() > 0 && !PluginBrowserTab.IsValid())
 	{
 		FNotificationInfo Info(LOCTEXT("NewPluginsPopupTitle", "New plugins are available"));
 		Info.bFireAndForget = true;
diff --git a/Engine/Plugins/Editor/PluginBrowser/Source/PluginBrowser/Private/PluginBrowserModule.h b/Engine/Plugins/Editor/PluginBrowser/Source/PluginBrowser/Private/PluginBrowserModule.h
--- a/Engine/Plugins/Editor/PluginBrowser/Source/PluginBrowser/Private/PluginBrowserModule.h
+++ b/Engine/Plugins/Editor/PluginBrowser/Source/PluginBrowser/Private/PluginBrowserModule.h
@@ -86,7 +86,7 @@ private:
 	TSharedRef<SDockTab> HandleSpawnPluginCreatorTab(const FSpawnTabArgs& SpawnTabArgs);
 
 	/** Callback for the main frame finishing load */
-	void OnMainFrameLoaded(TSharedPtr<SWindow> InRootWindow, bool bIsNewProjectWindow);
+	void OnMainFrameLoaded(TSharedPtr<SWindow> InRootWindow, bool bIsRunningStartupDialog);
 
 	/** Callback for when the user selects to edit installed plugins */
 	void OnNewPluginsPopupSettingsClicked();
diff --git a/Engine/Plugins/EnhancedInput/Source/InputEditor/Private/InputEditorModule.cpp b/Engine/Plugins/EnhancedInput/Source/InputEditor/Private/InputEditorModule.cpp
--- a/Engine/Plugins/EnhancedInput/Source/InputEditor/Private/InputEditorModule.cpp
+++ b/Engine/Plugins/EnhancedInput/Source/InputEditor/Private/InputEditorModule.cpp
@@ -425,7 +425,7 @@ void FInputEditorModule::ShutdownModule()
 	}
 }
 
-void FInputEditorModule::OnMainFrameCreationFinished(TSharedPtr<SWindow> InRootWindow, bool bIsNewProjectWindow)
+void FInputEditorModule::OnMainFrameCreationFinished(TSharedPtr<SWindow> InRootWindow, bool bIsRunningStartupDialog)
 {
 	AutoUpgradeDefaultInputClasses();
 }
diff --git a/Engine/Plugins/EnhancedInput/Source/InputEditor/Public/InputEditorModule.h b/Engine/Plugins/EnhancedInput/Source/InputEditor/Public/InputEditorModule.h
--- a/Engine/Plugins/EnhancedInput/Source/InputEditor/Public/InputEditorModule.h
+++ b/Engine/Plugins/EnhancedInput/Source/InputEditor/Public/InputEditorModule.h
@@ -43,7 +43,7 @@ private:
 		CreatedAssetTypeActions.Add(Action);
 	}
 
-	void OnMainFrameCreationFinished(TSharedPtr<SWindow> InRootWindow, bool bIsNewProjectWindow);
+	void OnMainFrameCreationFinished(TSharedPtr<SWindow> InRootWindow, bool bIsRunningStartupDialog);
 	
 	/** Automatically upgrade the current project to use Enhanced Input if it is currently set to the legacy input classes. */
 	void AutoUpgradeDefaultInputClasses();
diff --git a/Engine/Plugins/Media/BinkMedia/Source/BinkMediaPlayer/Private/BinkMediaPlayerModule.cpp b/Engine/Plugins/Media/BinkMedia/Source/BinkMediaPlayer/Private/BinkMediaPlayerModule.cpp
--- a/Engine/Plugins/Media/BinkMedia/Source/BinkMediaPlayer/Private/BinkMediaPlayerModule.cpp
+++ b/Engine/Plugins/Media/BinkMedia/Source/BinkMediaPlayer/Private/BinkMediaPlayerModule.cpp
@@ -222,7 +222,8 @@ struct FBinkMediaPlayerModule : IModuleInterface, FTickableGameObject
 	}
 
 #if BINKPLUGIN_UE4_EDITOR
-	void Initialize(TSharedPtr<SWindow> InRootWindow, bool bIsNewProjectWindow) {
+	void Initialize(TSharedPtr<SWindow> InRootWindow, bool bIsRunningStartupDialog)
+	{
 		// This overrides the 
 		FPropertyEditorModule& PropertyModule = FModuleManager::LoadModuleChecked<FPropertyEditorModule>("PropertyEditor");
 		PropertyModule.RegisterCustomClassLayout("MoviePlayerSettings", FOnGetDetailCustomizationInstance::CreateStatic(&FBinkMoviePlayerSettingsDetails::MakeInstance));
diff --git a/Engine/Plugins/Media/PixelStreaming/Source/PixelStreamingEditor/Private/PixelStreamingEditorModule.cpp b/Engine/Plugins/Media/PixelStreaming/Source/PixelStreamingEditor/Private/PixelStreamingEditorModule.cpp
--- a/Engine/Plugins/Media/PixelStreaming/Source/PixelStreamingEditor/Private/PixelStreamingEditorModule.cpp
+++ b/Engine/Plugins/Media/PixelStreaming/Source/PixelStreamingEditor/Private/PixelStreamingEditorModule.cpp
@@ -97,7 +97,7 @@ void FPixelStreamingEditorModule::InitEditorStreaming(IPixelStreamingModule& Mod
 		Streamer->SetSignallingServerURL(Module.GetDefaultSignallingURL());
 	}
 
-	IMainFrameModule::Get().OnMainFrameCreationFinished().AddLambda([&](TSharedPtr<SWindow> RootWindow, bool bIsNewProjectWindow)
+	IMainFrameModule::Get().OnMainFrameCreationFinished().AddLambda([&](TSharedPtr<SWindow> RootWindow, bool bIsRunningStartupDialog)
 	{
 		MaybeResizeEditor(RootWindow);
 
diff --git a/Engine/Source/Editor/Blutility/Private/EditorUtilitySubsystem.cpp b/Engine/Source/Editor/Blutility/Private/EditorUtilitySubsystem.cpp
--- a/Engine/Source/Editor/Blutility/Private/EditorUtilitySubsystem.cpp
+++ b/Engine/Source/Editor/Blutility/Private/EditorUtilitySubsystem.cpp
@@ -110,7 +110,7 @@ void UEditorUtilitySubsystem::AddReferencedObjects(UObject* InThis, FReferenceCo
 	}
 }
 
-void UEditorUtilitySubsystem::MainFrameCreationFinished(TSharedPtr<SWindow> InRootWindow, bool bIsNewProjectWindow)
+void UEditorUtilitySubsystem::MainFrameCreationFinished(TSharedPtr<SWindow> InRootWindow, bool bIsRunningStartupDialog)
 {
 	HandleStartup();
 }
diff --git a/Engine/Source/Editor/Blutility/Public/EditorUtilitySubsystem.h b/Engine/Source/Editor/Blutility/Public/EditorUtilitySubsystem.h
--- a/Engine/Source/Editor/Blutility/Public/EditorUtilitySubsystem.h
+++ b/Engine/Source/Editor/Blutility/Public/EditorUtilitySubsystem.h
@@ -48,7 +48,7 @@ public:
 
 	static void AddReferencedObjects(UObject* InThis, FReferenceCollector& Collector);
 
-	void MainFrameCreationFinished(TSharedPtr<SWindow> InRootWindow, bool bIsNewProjectWindow);
+	void MainFrameCreationFinished(TSharedPtr<SWindow> InRootWindow, bool bIsRunningStartupDialog);
 	void HandleStartup();
 
 	UPROPERTY(config)
diff --git a/Engine/Source/Editor/MainFrame/Private/MainFrameModule.cpp b/Engine/Source/Editor/MainFrame/Private/MainFrameModule.cpp
--- a/Engine/Source/Editor/MainFrame/Private/MainFrameModule.cpp
+++ b/Engine/Source/Editor/MainFrame/Private/MainFrameModule.cpp
@@ -32,6 +32,7 @@
 #include "Toolkits/FConsoleCommandExecutor.h"
 #include "Misc/MessageDialog.h"
 #include "ProfilingDebugging/StallDetector.h"
+#include "Interfaces/IEditorMainFrameProvider.h"
 
 DEFINE_LOG_CATEGORY(LogMainFrame);
 #define LOCTEXT_NAMESPACE "FMainFrameModule"
@@ -67,6 +68,52 @@ const FText StaticGetApplicationTitle( const bool bIncludeGameName )
 	return ApplicationTitle;
 }
 
+/* FProjectDialogProvider
+ *****************************************************************************/
+
+void FProjectDialogProvider::Register()
+{
+	IModularFeatures::Get().RegisterModularFeature(IEditorMainFrameProvider::GetModularFeatureName(), this);
+}
+
+void FProjectDialogProvider::UnRegister()
+{
+	IModularFeatures::Get().UnregisterModularFeature(IEditorMainFrameProvider::GetModularFeatureName(), this);
+}
+
+bool FProjectDialogProvider::IsRequestingMainFrameControl() const
+{
+	return !FApp::HasProjectName();
+}
+
+FMainFrameWindowOverrides FProjectDialogProvider::GetDesiredWindowConfiguration() const
+{
+	FMainFrameWindowOverrides ProjectDialogWindowConfig;
+
+	// Force tabs restored from layout that have no window (the LevelEditor tab) to use a docking area with
+	// embedded title area content.  We need to override the behavior here because we're creating the actual
+	// window ourselves instead of letting the tab management system create it for us.
+	ProjectDialogWindowConfig.bEmbedTitleAreaContent = false;
+
+	// Do not maximize the window initially. Keep a small dialog feel.
+	ProjectDialogWindowConfig.bInitiallyMaximized = false;
+
+	ProjectDialogWindowConfig.WindowSize = FMainFrameModule::GetProjectBrowserWindowSize();
+
+	ProjectDialogWindowConfig.bIsUserSizable = true;
+	ProjectDialogWindowConfig.bSupportsMaximize = true;
+	ProjectDialogWindowConfig.bSupportsMinimize = true;
+	ProjectDialogWindowConfig.CenterRules = EAutoCenter::PreferredWorkArea;
+	// When opening the project dialog, show "Project Browser" in the window title
+	ProjectDialogWindowConfig.WindowTitle = LOCTEXT("ProjectBrowserDialogTitle", "Unreal Project Browser");
+
+	return ProjectDialogWindowConfig;
+}
+
+TSharedRef<SWidget> FProjectDialogProvider::CreateMainFrameContentWidget() const
+{
+	return FGameProjectGenerationModule::Get().CreateGameProjectDialog(/*bAllowProjectOpening=*/true, /*bAllowProjectCreate=*/true);
+}
 
 /* IMainFrameModule implementation
  *****************************************************************************/
@@ -101,32 +148,40 @@ void FMainFrameModule::CreateDefaultMainFrameAuxiliary(const bool bStartImmersiv
 	{
 		FRootWindowLocation DefaultWindowLocation;
 
-		bool bEmbedTitleAreaContent = true;
-		bool bIsUserSizable = true;
-		bool bSupportsMaximize = true;
-		bool bSupportsMinimize = true;
-		EAutoCenter CenterRules = EAutoCenter::None;
-		FText WindowTitle;
-		if ( ShouldShowProjectDialogAtStartup() )
+		FMainFrameWindowOverrides WindowConfig;
+		bool bShowStartupDialogInPlaceOfMainEditor = false;
+		TSharedPtr<SWidget> MainFrameContent;
+
+		FGameProjectGenerationModule::Get();
+
+		TArray<IEditorMainFrameProvider*> MainFrameProviders = IModularFeatures::Get().GetModularFeatureImplementations<IEditorMainFrameProvider>(IEditorMainFrameProvider::GetModularFeatureName());		
+		for (IEditorMainFrameProvider* Provider: MainFrameProviders)
 		{
-			// Force tabs restored from layout that have no window (the LevelEditor tab) to use a docking area with
-			// embedded title area content.  We need to override the behavior here because we're creating the actual
-			// window ourselves instead of letting the tab management system create it for us.
-			bEmbedTitleAreaContent = false;
-
-			// Do not maximize the window initially. Keep a small dialog feel.
-			DefaultWindowLocation.InitiallyMaximized = false;
-
-			DefaultWindowLocation.WindowSize = GetProjectBrowserWindowSize();
-
-			bIsUserSizable = true;
-			bSupportsMaximize = true;
-			bSupportsMinimize = true;
-			CenterRules = EAutoCenter::PreferredWorkArea;;
-			// When opening the project dialog, show "Project Browser" in the window title
-			WindowTitle = LOCTEXT("ProjectBrowserDialogTitle", "Unreal Project Browser");
+			if (Provider && Provider->IsRequestingMainFrameControl())
+			{
+				WindowConfig = Provider->GetDesiredWindowConfiguration();
+				MainFrameContent = Provider->CreateMainFrameContentWidget();
+				bShowStartupDialogInPlaceOfMainEditor = true;
+
+				// We can only have one main frame provider, which means we just use  the first active one
+				break;
+			}
 		}
-		else
+
+		if (WindowConfig.ScreenPosition.IsSet())
+		{
+			DefaultWindowLocation.ScreenPosition = WindowConfig.ScreenPosition.GetValue();
+		}
+		if (WindowConfig.WindowSize.IsSet())
+		{
+			DefaultWindowLocation.WindowSize = WindowConfig.WindowSize.GetValue();
+		}
+		if (WindowConfig.bInitiallyMaximized.IsSet())
+		{
+			DefaultWindowLocation.InitiallyMaximized = WindowConfig.bInitiallyMaximized.GetValue();
+		}
+
+		if (!bShowStartupDialogInPlaceOfMainEditor)
 		{
 			if( bStartImmersive )
 			{
@@ -135,19 +190,19 @@ void FMainFrameModule::CreateDefaultMainFrameAuxiliary(const bool bStartImmersiv
 			}
 
 			const bool bIncludeGameName = true;
-			WindowTitle = GetApplicationTitle( bIncludeGameName );
+			WindowConfig.WindowTitle = GetApplicationTitle( bIncludeGameName );
 		}
 
 		TSharedRef<SWindow> RootWindow = SNew(SWindow)
-			.AutoCenter(CenterRules)
-			.Title( WindowTitle )
+			.AutoCenter(WindowConfig.CenterRules)
+			.Title( WindowConfig.WindowTitle )
 			.IsInitiallyMaximized( DefaultWindowLocation.InitiallyMaximized )
 			.ScreenPosition( DefaultWindowLocation.ScreenPosition )
 			.ClientSize( DefaultWindowLocation.WindowSize )
-			.CreateTitleBar( !bEmbedTitleAreaContent )
-			.SizingRule( bIsUserSizable ? ESizingRule::UserSized : ESizingRule::FixedSize )
-			.SupportsMaximize( bSupportsMaximize )
-			.SupportsMinimize( bSupportsMinimize );
+			.CreateTitleBar( !WindowConfig.bEmbedTitleAreaContent )
+			.SizingRule( WindowConfig.bIsUserSizable ? ESizingRule::UserSized : ESizingRule::FixedSize )
+			.SupportsMaximize( WindowConfig.bSupportsMaximize )
+			.SupportsMinimize( WindowConfig.bSupportsMinimize );
 
 		const bool bShowRootWindowImmediately = false;
 		FSlateApplication::Get().AddWindow( RootWindow, bShowRootWindowImmediately );
@@ -157,13 +212,8 @@ void FMainFrameModule::CreateDefaultMainFrameAuxiliary(const bool bStartImmersiv
 
 		FSlateNotificationManager::Get().SetRootWindow(RootWindow);
 
-		TSharedPtr<SWidget> MainFrameContent;
 		bool bLevelEditorIsMainTab = false;
-		if ( ShouldShowProjectDialogAtStartup() )
-		{
-			MainFrameContent = FGameProjectGenerationModule::Get().CreateGameProjectDialog(/*bAllowProjectOpening=*/true, /*bAllowProjectCreate=*/true);
-		}
-		else
+		if (!bShowStartupDialogInPlaceOfMainEditor)
 		{
 			// Get desktop metrics
 			FDisplayMetrics DisplayMetrics;
@@ -262,7 +312,7 @@ void FMainFrameModule::CreateDefaultMainFrameAuxiliary(const bool bStartImmersiv
 			FToolMenuContext EmptyContext;
 			MakeMainMenu(FGlobalTabmanager::Get(), "MainFrame.NomadMainMenu", EmptyContext);
 		
-			MainFrameContent = FGlobalTabmanager::Get()->RestoreFrom(LoadedLayout, RootWindow, bEmbedTitleAreaContent, OutputCanBeNullptr);
+			MainFrameContent = FGlobalTabmanager::Get()->RestoreFrom(LoadedLayout, RootWindow, WindowConfig.bEmbedTitleAreaContent, OutputCanBeNullptr);
 			// MainFrameContent will only be nullptr if its main area contains invalid tabs (probably some layout bug). If so, reset layout to avoid potential crashes
 			if (!MainFrameContent.IsValid())
 			{
@@ -319,7 +369,7 @@ void FMainFrameModule::CreateDefaultMainFrameAuxiliary(const bool bStartImmersiv
 		MRUFavoritesList = new FMainMRUFavoritesList;
 		MRUFavoritesList->ReadFromINI();
 
-		MainFrameCreationFinishedEvent.Broadcast(RootWindow, ShouldShowProjectDialogAtStartup());
+		MainFrameCreationFinishedEvent.Broadcast(RootWindow, bShowStartupDialogInPlaceOfMainEditor);
 	}
 }
 
@@ -534,6 +584,8 @@ void FMainFrameModule::StartupModule( )
 
 	MRUFavoritesList = NULL;
 
+	ProjectDialogProvider.Register();
+
 	ensureMsgf(!IsRunningGame(), TEXT("The MainFrame module should only be loaded when running the editor.  Code that extends the editor, adds menu items, etc... should not run when running in -game mode or in a non-WITH_EDITOR build"));
 	MainFrameHandler = MakeShareable(new FMainFrameHandler);
 
@@ -597,6 +649,8 @@ void FMainFrameModule::ShutdownModule( )
 		SourceCodeAccessModule.OnOpenFileFailed().RemoveAll( this );
 	}
 #endif
+
+	ProjectDialogProvider.UnRegister();
 }
 
 
@@ -616,16 +670,6 @@ void FMainFrameModule::HandleResizeMainFrameCommand(const TArray<FString>& Args)
 	
 }
 
-
-/* FMainFrameModule implementation
- *****************************************************************************/
-
-bool FMainFrameModule::ShouldShowProjectDialogAtStartup( ) const
-{
-	return !FApp::HasProjectName();
-}
-
-
 /* FMainFrameModule event handlers
  *****************************************************************************/
 
diff --git a/Engine/Source/Editor/MainFrame/Private/MainFrameModule.h b/Engine/Source/Editor/MainFrame/Private/MainFrameModule.h
--- a/Engine/Source/Editor/MainFrame/Private/MainFrameModule.h
+++ b/Engine/Source/Editor/MainFrame/Private/MainFrameModule.h
@@ -11,6 +11,24 @@
 #include "UnrealEdMisc.h"
 #include "Frame/MainFrameHandler.h"
 #include "Misc/CompilationResult.h"
+#include "Interfaces/IEditorMainFrameProvider.h"
+
+/** 
+ * Utility class, which hooks into the editor's main window startup and 
+ * supplies a project dialog window if the project hasn't been set. 
+ */
+class FProjectDialogProvider : public IEditorMainFrameProvider
+{
+public:
+	void Register();
+	void UnRegister();
+
+	//~ Begin IEditorMainFrameProvider interface
+	virtual bool IsRequestingMainFrameControl() const override;
+	virtual FMainFrameWindowOverrides GetDesiredWindowConfiguration() const override;
+	virtual TSharedRef<SWidget> CreateMainFrameContentWidget() const override;
+	//~ End IEditorMainFrameProvider interface
+};
 
 /**
  * Editor main frame module
@@ -196,16 +214,6 @@ public:
 	}
 
 	static void HandleResizeMainFrameCommand(const TArray<FString>& Args);
-protected:
-
-	/**
-	 * Checks whether the project dialog should be shown at startup.
-	 *
-	 * The project dialog should be shown if the Editor was started without a game specified.
-	 *
-	 * @return true if the project dialog should be shown, false otherwise.
-	 */
-	bool ShouldShowProjectDialogAtStartup( ) const;
 
 public:
 
@@ -292,4 +300,7 @@ private:
 
 	// Is recreating Default Main Frame
 	bool bRecreatingDefaultMainFrame;
+
+	// Instantiation of the object responsible for spawning the editor's project dialog on startup
+	FProjectDialogProvider ProjectDialogProvider;
 };
diff --git a/Engine/Source/Editor/MainFrame/Public/Interfaces/IEditorMainFrameProvider.h b/Engine/Source/Editor/MainFrame/Public/Interfaces/IEditorMainFrameProvider.h
new file mode 100644
--- /dev/null
+++ b/Engine/Source/Editor/MainFrame/Public/Interfaces/IEditorMainFrameProvider.h
@@ -0,0 +1,73 @@
+// Copyright Epic Games, Inc. All Rights Reserved.
+
+#pragma once
+
+#include "CoreMinimal.h"
+#include "Features/IModularFeature.h"
+#include "UObject/NameTypes.h"
+#include "Templates/SharedPointer.h"
+#include "Widgets/SWindow.h" // for EAutoCenter
+
+class SWidget;
+
+/** Set of configuration settings for the editor's main window. */
+struct FMainFrameWindowOverrides
+{
+	FText WindowTitle;
+
+	// If not set, will fallback to main Editor defaults
+	TOptional<FVector2D> ScreenPosition;
+	TOptional<FVector2D> WindowSize;
+	TOptional<bool> bInitiallyMaximized;
+
+	EAutoCenter CenterRules = EAutoCenter::None;
+
+	bool bEmbedTitleAreaContent : 1;
+	bool bIsUserSizable : 1;
+	bool bSupportsMaximize : 1;
+	bool bSupportsMinimize : 1;
+
+	FMainFrameWindowOverrides()
+		: bEmbedTitleAreaContent(true)
+		, bIsUserSizable(true)
+		, bSupportsMaximize(true)
+		, bSupportsMinimize(true)
+	{}
+};
+
+/** 
+ * Base feature class, which serves as a hook for external sources to override 
+ * the editor's main window on startup. 
+ */
+class IEditorMainFrameProvider : public IModularFeature
+{
+public:
+	virtual ~IEditorMainFrameProvider() {}
+
+	static FName GetModularFeatureName()
+	{
+		static FName FeatureName = FName(TEXT("EditorMainFrameProvider"));
+		return FeatureName;
+	}
+
+	/**
+	 * Signals if this provider wishes to override the editor's main frame on 
+	 * editor startup.
+	 * NOTE: When true, this will block the standard editor window from appearing 
+	 *       (instead replacing it with the provider's widget). 
+	 */
+	virtual bool IsRequestingMainFrameControl() const = 0;
+
+	/** 
+	 * Specifies the desired window title, dimensions, etc. -- for the main frame 
+	 * window (to fit the provier's widget).
+	 */
+	virtual FMainFrameWindowOverrides GetDesiredWindowConfiguration() const = 0;
+
+	/**
+	 * Spawns a widget for slotting into the main editor window (in-place of the 
+	 * standard editor).
+	 * @return A replacement widget for the editor's main window.
+	 */
+	virtual TSharedRef<SWidget> CreateMainFrameContentWidget() const = 0;
+};
\ No newline at end of file
diff --git a/Engine/Source/Editor/StatusBar/Private/StatusBarSubsystem.cpp b/Engine/Source/Editor/StatusBar/Private/StatusBarSubsystem.cpp
--- a/Engine/Source/Editor/StatusBar/Private/StatusBarSubsystem.cpp
+++ b/Engine/Source/Editor/StatusBar/Private/StatusBarSubsystem.cpp
@@ -696,9 +696,9 @@ void UStatusBarSubsystem::CreateContentBrowserIfNeeded()
 	}
 }
 
-void UStatusBarSubsystem::CreateAndShowNewUserTipIfNeeded(TSharedPtr<SWindow> ParentWindow, bool bIsNewProjectDialog)
+void UStatusBarSubsystem::CreateAndShowNewUserTipIfNeeded(TSharedPtr<SWindow> ParentWindow, bool bIsRunningStartupDialog)
 {
-	if(!bIsNewProjectDialog)
+	if(!bIsRunningStartupDialog)
 	{
 		FString CurrentState = GetNewUserTipState();
 
@@ -735,13 +735,13 @@ const FString UStatusBarSubsystem::GetNewUserTipState() const
 	return CurrentState;
 }
 
-void UStatusBarSubsystem::CreateAndShowOneTimeIndustryQueryIfNeeded(TSharedPtr<SWindow> ParentWindow, bool bIsNewProjectDialog)
+void UStatusBarSubsystem::CreateAndShowOneTimeIndustryQueryIfNeeded(TSharedPtr<SWindow> ParentWindow, bool bIsRunningStartupDialog)
 {
 	// Only show for external builds where editor analytics are on and the industry popup is not suppressed
 	if(SOneTimeIndustryQuery::ShouldShowIndustryQuery())
 	{
 		FString NewUserTipState = GetNewUserTipState();
-		if (!bIsNewProjectDialog && NewUserTipState == TEXT("1"))
+		if (!bIsRunningStartupDialog && NewUserTipState == TEXT("1"))
 		{
 			const FString StoreId = TEXT("Epic Games");
 			const FString SectionName = TEXT("Unreal Engine/Editor");
diff --git a/Engine/Source/Editor/StatusBar/Public/StatusBarSubsystem.h b/Engine/Source/Editor/StatusBar/Public/StatusBarSubsystem.h
--- a/Engine/Source/Editor/StatusBar/Public/StatusBarSubsystem.h
+++ b/Engine/Source/Editor/StatusBar/Public/StatusBarSubsystem.h
@@ -164,9 +164,9 @@ private:
 	bool ToggleContentBrowser(TSharedRef<SWindow> ParentWindow);
 	void OnDebugConsoleClosed(TWeakPtr<SStatusBar> OwningStatusBar);
 	void CreateContentBrowserIfNeeded();
-	void CreateAndShowNewUserTipIfNeeded(TSharedPtr<SWindow> ParentWindow, bool bIsNewProjectDialog);
+	void CreateAndShowNewUserTipIfNeeded(TSharedPtr<SWindow> ParentWindow, bool bIsRunningStartupDialog);
 	const FString GetNewUserTipState() const;
-	void CreateAndShowOneTimeIndustryQueryIfNeeded(TSharedPtr<SWindow> ParentWindow, bool bIsNewProjectDialog);
+	void CreateAndShowOneTimeIndustryQueryIfNeeded(TSharedPtr<SWindow> ParentWindow, bool bIsRunningStartupDialog);
 	static const FString GetOneTimeStateWithFallback(const FString StoreId, const FString SectionName, const FString KeyName, const FString FallbackIniLocation, const FString FallbackIniKey);
 	static void SetOneTimeStateWithFallback(const FString StoreId, const FString SectionName, const FString KeyName, const FString FallbackIniLocation, const FString FallbackIniKey);
 
diff --git a/Engine/Source/Editor/UnrealEd/Private/EditorCommandLineUtils.cpp b/Engine/Source/Editor/UnrealEd/Private/EditorCommandLineUtils.cpp
--- a/Engine/Source/Editor/UnrealEd/Private/EditorCommandLineUtils.cpp
+++ b/Engine/Source/Editor/UnrealEd/Private/EditorCommandLineUtils.cpp
@@ -441,11 +441,11 @@ static void EditorCommandLineUtilsImpl::ForceCloseEditor()
 }
 
 //------------------------------------------------------------------------------
-static void EditorCommandLineUtilsImpl::RunAssetDiffCommand(TSharedPtr<SWindow> MainEditorWindow, bool bIsRunningProjBrowser, FString CommandArgs)
+static void EditorCommandLineUtilsImpl::RunAssetDiffCommand(TSharedPtr<SWindow> MainEditorWindow, bool bIsRunningStartupDialog, FString CommandArgs)
 {
-	// if the editor is running the project browser, then the user has to first 
-	// select a project (and then the editor will re-launch with this command).
-	if (bIsRunningProjBrowser) 
+	// if the editor is running the project browser (or another like dialog on startup), 
+	// then the user has to first make a selection (and then the editor will re-launch with this command).
+	if (bIsRunningStartupDialog)
 	{
 		// @TODO: can we run without loading a project?
 		return;
