Added an editor project option to hide the 'Submit Content' menu action. This menu item is unwanted in some projects as it's not running any validation before submitting.

diff --git a/Engine/Source/Editor/MainFrame/Private/Frame/MainFrameActions.cpp b/Engine/Source/Editor/MainFrame/Private/Frame/MainFrameActions.cpp
--- a/Engine/Source/Editor/MainFrame/Private/Frame/MainFrameActions.cpp
+++ b/Engine/Source/Editor/MainFrame/Private/Frame/MainFrameActions.cpp
@@ -109,7 +109,7 @@ void FMainFrameCommands::RegisterCommands()
 	ActionList->MapAction(ViewChangelists, FExecuteAction::CreateStatic( &FMainFrameActionCallbacks::ViewChangelists ), FCanExecuteAction::CreateStatic( &FMainFrameActionCallbacks::CanViewChangelists ) );
 
 	UI_COMMAND( SubmitContent, "Submit Content", "Opens a dialog with check in options for content and levels.", EUserInterfaceActionType::Button, FInputChord() );
-	ActionList->MapAction( SubmitContent, FExecuteAction::CreateLambda([]() { FSourceControlWindows::ChoosePackagesToCheckIn(); }), FCanExecuteAction::CreateStatic(&FSourceControlWindows::CanChoosePackagesToCheckIn ) );
+	ActionList->MapAction( SubmitContent, FExecuteAction::CreateLambda([]() { FSourceControlWindows::ChoosePackagesToCheckIn(); }), FCanExecuteAction::CreateStatic(&FSourceControlWindows::CanChoosePackagesToCheckIn ), FGetActionCheckState(), FIsActionButtonVisible::CreateStatic(FSourceControlWindows::ShouldChoosePackagesToCheckBeVisible) );
 
 	UI_COMMAND( ConnectToSourceControl, "Connect to Source Control...", "Connect to source control to allow source control operations to be performed on content and levels.", EUserInterfaceActionType::Button, FInputChord() );
 	ActionList->MapAction( ConnectToSourceControl, FExecuteAction::CreateStatic( &FMainFrameActionCallbacks::ConnectToSourceControl ), DefaultExecuteAction );
diff --git a/Engine/Source/Editor/SourceControlWindows/Private/SourceControlSettings.h b/Engine/Source/Editor/SourceControlWindows/Private/SourceControlSettings.h
--- a/Engine/Source/Editor/SourceControlWindows/Private/SourceControlSettings.h
+++ b/Engine/Source/Editor/SourceControlWindows/Private/SourceControlSettings.h
@@ -22,4 +22,7 @@ public:
 
 	UPROPERTY(config, EditAnywhere, Category="Source Control Changelist View")
 	bool bShowAssetCheckedOutByColumn = true;
+
+	UPROPERTY(EditAnywhere, config, Category="Source Control Tools")
+	bool bEnableSubmitContentMenuAction = true;
 };
diff --git a/Engine/Source/Editor/SourceControlWindows/Private/SourceControlWindows.cpp b/Engine/Source/Editor/SourceControlWindows/Private/SourceControlWindows.cpp
--- a/Engine/Source/Editor/SourceControlWindows/Private/SourceControlWindows.cpp
+++ b/Engine/Source/Editor/SourceControlWindows/Private/SourceControlWindows.cpp
@@ -13,6 +13,7 @@
 #include "Logging/TokenizedMessage.h"
 #include "Misc/MessageDialog.h"
 #include "Widgets/Notifications/SNotificationList.h"
+#include "SourceControlSettings.h"
 
 
 #if SOURCE_CONTROL_WITH_SLATE
@@ -143,7 +144,13 @@ bool FSourceControlWindows::CanChoosePackagesToCheckIn()
 	
 	return ISourceControlModule::Get().IsEnabled() &&
 		ISourceControlModule::Get().GetProvider().IsAvailable() &&
-		!ChoosePackagesToCheckInNotification.IsValid();
+		!ChoosePackagesToCheckInNotification.IsValid()
+		;
+}
+
+bool FSourceControlWindows::ShouldChoosePackagesToCheckBeVisible()
+{
+	return GetDefault<USourceControlSettings>()->bEnableSubmitContentMenuAction;
 }
 
 
diff --git a/Engine/Source/Editor/SourceControlWindows/Public/SourceControlWindows.h b/Engine/Source/Editor/SourceControlWindows/Public/SourceControlWindows.h
--- a/Engine/Source/Editor/SourceControlWindows/Public/SourceControlWindows.h
+++ b/Engine/Source/Editor/SourceControlWindows/Public/SourceControlWindows.h
@@ -53,6 +53,9 @@ public:
 	/** Determines whether we can choose packages to check in (we cant if an operation is already in progress) */
 	static bool CanChoosePackagesToCheckIn();
 
+	/** Determines if the Submit Content action should be visible or not */
+	static bool ShouldChoosePackagesToCheckBeVisible();
+
 	/**
 	 * Display check in dialog for the specified packages and get additional result information
 	 *
diff --git a/Engine/Source/Editor/StatusBar/Private/SourceControlMenuHelpers.cpp b/Engine/Source/Editor/StatusBar/Private/SourceControlMenuHelpers.cpp
--- a/Engine/Source/Editor/StatusBar/Private/SourceControlMenuHelpers.cpp
+++ b/Engine/Source/Editor/StatusBar/Private/SourceControlMenuHelpers.cpp
@@ -102,6 +102,12 @@ bool FSourceControlCommands::SubmitContent_IsVisible()
 	{
 		return false;
 	}
+
+	if (!FSourceControlWindows::ShouldChoosePackagesToCheckBeVisible())
+	{
+		return false;
+	}
+
 	return true;
 }
 
