Rewrite "Save Layout As..." using SGridPanel to improve text overflow behavior.

diff --git a/Engine/Source/Editor/MainFrame/Private/Menus/SaveLayoutDialog.cpp b/Engine/Source/Editor/MainFrame/Private/Menus/SaveLayoutDialog.cpp
--- a/Engine/Source/Editor/MainFrame/Private/Menus/SaveLayoutDialog.cpp
+++ b/Engine/Source/Editor/MainFrame/Private/Menus/SaveLayoutDialog.cpp
@@ -30,6 +30,7 @@
 #include "Widgets/Input/SEditableTextBox.h"
 #include "Widgets/Layout/SBorder.h"
 #include "Widgets/Layout/SBox.h"
+#include "Widgets/Layout/SGridPanel.h"
 #include "Widgets/SBoxPanel.h"
 #include "Widgets/SCompoundWidget.h"
 #include "Widgets/SWindow.h"
@@ -202,16 +203,11 @@ void SSaveLayoutDialog::Construct(const FArguments& InArgs, const TSharedRef<FSa
 		CurrentlyEnteredLayoutDescription = SaveLayoutDialogParams->LayoutDescriptions[0];
 	}
 
-	// UI part
-	// The root SWidget in this dialog
-	TSharedRef<SVerticalBox> MainVerticalBox = SNew(SVerticalBox);
-	TSharedRef<SVerticalBox> ContentBox = SNew(SVerticalBox);
-
-	// Add to ChildSlot
 	ChildSlot
 	[
 		SNew(SBorder)
 		.BorderImage(FAppStyle::Get().GetBrush("Brushes.Panel"))
+		.Padding(16.0f)
 		[
 			SNew(SBox)
 			.WidthOverride(600.f)
@@ -220,22 +216,27 @@ void SSaveLayoutDialog::Construct(const FArguments& InArgs, const TSharedRef<FSa
 				SNew(SVerticalBox)
 				+ SVerticalBox::Slot()
 				.AutoHeight()
-				.Padding(5.0f,2.0f)
 				[
-					SNew(SHorizontalBox)
-					+SHorizontalBox::Slot()
-					.FillWidth(0.6f)
+					SNew(SGridPanel)
+					.FillColumn(1, 1.0f)
+					+SGridPanel::Slot(0,0)
 					.HAlign(HAlign_Right)
 					.VAlign(VAlign_Center)
-					.Padding(5.0f, 2.0f)
-					[	
+					[
 						SNew(STextBlock)
 						.Text(LOCTEXT("LayoutNameBoxLabel", "Name"))
+						.Margin(FMargin(0, 0, 8, 8))
 					]
-					+SHorizontalBox::Slot()
-					.FillWidth(2.0f)
+					+SGridPanel::Slot(0,1)
+					.HAlign(HAlign_Right)
 					.VAlign(VAlign_Center)
-					.Padding(5.0f, 2.0f)
+					[
+						SNew(STextBlock)
+						.Text(LOCTEXT("LayoutDescriptionBoxLabel", "Description"))
+						.Margin(FMargin(0, 0, 8, 8))
+					]
+					+SGridPanel::Slot(1,0)
+					.Padding(0, 0, 0, 8)
 					[
 						SNew(SEditableTextBox)
 						.Text(CurrentlyEnteredLayoutName)
@@ -243,28 +244,12 @@ void SSaveLayoutDialog::Construct(const FArguments& InArgs, const TSharedRef<FSa
 						.OnTextChanged(this, &SSaveLayoutDialog::OnLayoutNameTextCommited, ETextCommit::Default)
 						.SelectAllTextWhenFocused(true)
 					]
-				]
-				+ SVerticalBox::Slot()
-				.Padding(5.0f, 2.0f)
-				.AutoHeight()
-				[
-					SNew(SHorizontalBox)
-					+ SHorizontalBox::Slot()
-					.FillWidth(.6f)
-					.HAlign(HAlign_Right)
-					.VAlign(VAlign_Center)
-					.Padding(5.0f, 2.0f)
-					[
-						SNew(STextBlock)
-						.Text(LOCTEXT("LayoutDescriptionBoxLabel", "Description (Optional)"))
-					]
-					+ SHorizontalBox::Slot()
-					.VAlign(VAlign_Center)
-					.FillWidth(2.0f)
-					.Padding(5.0f, 2.0f)
+					+SGridPanel::Slot(1,1)
+					.Padding(0, 0, 0, 8)
 					[
 						SNew(SEditableTextBox)
 						.Text(CurrentlyEnteredLayoutDescription)
+						.HintText(LOCTEXT("LayoutDescriptionInputBoxHintText", "Optional"))
 						.OnTextCommitted(this, &SSaveLayoutDialog::OnLayoutDescriptionTextCommited)
 						.OnTextChanged(this, &SSaveLayoutDialog::OnLayoutDescriptionTextCommited, ETextCommit::Default)
 						.SelectAllTextWhenFocused(true)
@@ -291,13 +276,12 @@ void SSaveLayoutDialog::Construct(const FArguments& InArgs, const TSharedRef<FSa
 				+ SVerticalBox::Slot()
 				.HAlign(HAlign_Right)
 				.VAlign(VAlign_Bottom)
-				.Padding(0, 0, 5, 5)
 				[
 					SNew(SHorizontalBox)
 					+ SHorizontalBox::Slot()
 					.AutoWidth()
 					.VAlign(VAlign_Bottom)
-					.Padding(4, 3)
+					.Padding(0,0,8,0)
 					[
 						SNew(SPrimaryButton)
 						.Text(LOCTEXT("SaveLayoutDialogSaveButton", "Save"))
@@ -307,7 +291,6 @@ void SSaveLayoutDialog::Construct(const FArguments& InArgs, const TSharedRef<FSa
 					+ SHorizontalBox::Slot()
 					.AutoWidth()
 					.VAlign(VAlign_Bottom)
-					.Padding(4, 3)
 					[
 						SNew(SButton)
 						.Text(LOCTEXT("SaveLayoutDialogCancelButton", "Cancel"))
