Properly disable the "Restore open asset editors" prompt when the override is set

diff --git a/Engine/Source/Editor/MainFrame/Private/Menus/LayoutsMenu.cpp b/Engine/Source/Editor/MainFrame/Private/Menus/LayoutsMenu.cpp
--- a/Engine/Source/Editor/MainFrame/Private/Menus/LayoutsMenu.cpp
+++ b/Engine/Source/Editor/MainFrame/Private/Menus/LayoutsMenu.cpp
@@ -797,10 +797,12 @@ void FLayoutsMenuLoad::ReloadCurrentLayout()
 	// Disable config saving
 	UAssetEditorSubsystem* AssetEditorSubsystem = (GEditor ? GEditor->GetEditorSubsystem<UAssetEditorSubsystem>() : nullptr);
 	const bool bAreAssetEditorOpened = (AssetEditorSubsystem ? AssetEditorSubsystem->GetAllEditedAssets().Num() > 0 : false); // Are there asset editors opened?
+	TOptional<bool> bAutoRestoreAndDisableSavingOverride;
 	if (bAreAssetEditorOpened)
 	{
 		// Save open asset editors + disable manual saving (no need to close them with AssetEditorSubsystem->CloseAllAssetEditors)
 		AssetEditorSubsystem->SaveOpenAssetEditors(true);
+		bAutoRestoreAndDisableSavingOverride = AssetEditorSubsystem->GetAutoRestoreAndDisableSavingOverride();
 		AssetEditorSubsystem->SetAutoRestoreAndDisableSavingOverride(true);
 	}
 	FUnrealEdMisc::Get().AllowSavingLayoutOnClose(false);
@@ -813,7 +815,7 @@ void FLayoutsMenuLoad::ReloadCurrentLayout()
 		check(AssetEditorSubsystem);
 		// Restore asset editors + disable manual saving and avoid trying to re-open the asset editors twice
 		AssetEditorSubsystem->RestorePreviouslyOpenAssets();
-		AssetEditorSubsystem->SetAutoRestoreAndDisableSavingOverride(TOptional<bool>());
+		AssetEditorSubsystem->SetAutoRestoreAndDisableSavingOverride(bAutoRestoreAndDisableSavingOverride);
 	}
 	// Save layout and create duplicated ini file (DuplicatedEditorLayoutIniFilePath)
 	FPrivateLayoutsMenu::SaveLayoutWithoutRemovingTempLayoutFiles();
diff --git a/Engine/Source/Editor/UnrealEd/Private/Subsystems/AssetEditorSubsystem.cpp b/Engine/Source/Editor/UnrealEd/Private/Subsystems/AssetEditorSubsystem.cpp
--- a/Engine/Source/Editor/UnrealEd/Private/Subsystems/AssetEditorSubsystem.cpp
+++ b/Engine/Source/Editor/UnrealEd/Private/Subsystems/AssetEditorSubsystem.cpp
@@ -840,6 +840,11 @@ void UAssetEditorSubsystem::SetAutoRestoreAndDisableSavingOverride(TOptional<boo
 
 }
 
+TOptional<bool> UAssetEditorSubsystem::GetAutoRestoreAndDisableSavingOverride() const
+{
+	return bAutoRestoreAndDisableSaving;
+}
+
 void UAssetEditorSubsystem::SpawnRestorePreviouslyOpenAssetsNotification(const bool bCleanShutdown, const TArray<FString>& AssetsToOpen)
 {
 	FText NotificationMessage = bCleanShutdown
diff --git a/Engine/Source/Editor/UnrealEd/Public/Subsystems/AssetEditorSubsystem.h b/Engine/Source/Editor/UnrealEd/Public/Subsystems/AssetEditorSubsystem.h
--- a/Engine/Source/Editor/UnrealEd/Public/Subsystems/AssetEditorSubsystem.h
+++ b/Engine/Source/Editor/UnrealEd/Public/Subsystems/AssetEditorSubsystem.h
@@ -257,6 +257,9 @@ public:
 	/** Set/Unset the bAutoRestoreAndDisableSaving override, along with setting bRequestRestorePreviouslyOpenAssets to false to avoid running RestorePreviouslyOpenAssets() twice */
 	void SetAutoRestoreAndDisableSavingOverride(TOptional<bool> bInAutoRestoreAndDisableSaving);
 
+	/** Get the current value of bAutoRestoreAndDisableSaving */
+	TOptional<bool> GetAutoRestoreAndDisableSavingOverride() const;
+
 private:
 
 	/** Handles a package being reloaded */
