Fixed warnings when the layout menu is accessing the config cache:
- Normalize the layout file paths as the config cache expects them to be.
- Removed the ConvertToAbsolutePathForExternalAppForRead call since the config cache expects relatives paths and it generates an unusable path if the platform file is IPlatformPakFile.

#tests Validated that this works with cooked editors that use the IPlatformPakFile.

diff --git a/Engine/Source/Editor/MainFrame/Private/Menus/LayoutsMenu.cpp b/Engine/Source/Editor/MainFrame/Private/Menus/LayoutsMenu.cpp
--- a/Engine/Source/Editor/MainFrame/Private/Menus/LayoutsMenu.cpp
+++ b/Engine/Source/Editor/MainFrame/Private/Menus/LayoutsMenu.cpp
@@ -362,13 +362,7 @@ void FPrivateLayoutsMenu::DisplayLayoutsInternal(FToolMenuSection& InSection, co
 	// If there are Layout ini files, read them
 	for (int32 LayoutIndex = 0; LayoutIndex < InLayoutIniFileNames.Num(); ++LayoutIndex)
 	{
-		FString LayoutFilePath = IFileManager::Get().ConvertToAbsolutePathForExternalAppForRead(
-			*FPaths::Combine(InLayoutsDirectory, InLayoutIniFileNames[LayoutIndex]));
-
-		if (LayoutFilePath.StartsWith(TEXT("Pak: ")))
-		{
-			LayoutFilePath = FPaths::Combine(InLayoutsDirectory, InLayoutIniFileNames[LayoutIndex]);
-		}
+		FString LayoutFilePath = FConfigCacheIni::NormalizeConfigIniPath(*FPaths::Combine(InLayoutsDirectory, InLayoutIniFileNames[LayoutIndex]));
 
 		// Make sure it is a layout file
 		GConfig->UnloadFile(LayoutFilePath); // We must re-read it to avoid the Editor to use a previously cached name and description
