Use ToolMenus to add tab spawners under Edit, Window, and Tools instead of FMenuBuilder.

This change only affects the main menu entries added by MainMenu.cpp under "MainFrame.MainMenu.Edit", "MainFrame.MainMenu.Window", and "MainFrame.MainMenu.Tools". See the header "Acceptance Criteria WM-1" in https://docs.google.com/document/d/1xZsXdaa3CgIYYmKbZzEUuAbnjgqbV53Qdtax_RbVBLs for an example of what can change in these dynamic menus.

diff --git a/Engine/Source/Editor/MainFrame/Private/Menus/MainMenu.cpp b/Engine/Source/Editor/MainFrame/Private/Menus/MainMenu.cpp
--- a/Engine/Source/Editor/MainFrame/Private/Menus/MainMenu.cpp
+++ b/Engine/Source/Editor/MainFrame/Private/Menus/MainMenu.cpp
@@ -31,9 +31,121 @@
 #include "Framework/Commands/GenericCommands.h"
 #include "ITranslationEditor.h"
 #include "LauncherPlatformModule.h"
+#include "MainFrameLog.h"
 
 #define LOCTEXT_NAMESPACE "MainFileMenu"
 
+namespace UE::MainMenu::Private
+{
+
+void MakeSpawnerToolMenuEntry(TSharedPtr<FTabManager> TabManager, FToolMenuSection& InSection, const TSharedPtr<FTabSpawnerEntry> &InSpawner)
+{
+	if (!InSpawner->IsHidden())
+	{
+		InSection.AddMenuEntry(
+			NAME_None,
+			InSpawner->GetDisplayName().IsEmpty() ? FText::FromName(InSpawner->GetTabType()) : InSpawner->GetDisplayName(),
+			InSpawner->GetTooltipText(),
+			InSpawner->GetIcon(),
+			TabManager->GetUIActionForTabSpawnerMenuEntry(InSpawner),
+			EUserInterfaceActionType::Check
+		);
+	}
+}
+
+void PopulateTabSpawnerToolMenu_Helper(UToolMenu* InMenu, TSharedPtr<FTabManager> TabManager, TSharedRef<FWorkspaceItem> InMenuStructure, TSharedRef< TArray< TWeakPtr<FTabSpawnerEntry> > > AllSpawners, const int32 RecursionLevel, bool PutLabelOnSection)
+{
+	FToolMenuSection* TheSection = nullptr;
+	if (PutLabelOnSection)
+	{
+		const FText Label = InMenuStructure->GetDisplayName();
+		const FName Name = MakeObjectNameFromDisplayLabel(Label.ToString(), FName(TEXT("SectionName")));
+		TheSection = &InMenu->FindOrAddSection(Name);
+		TheSection->Label = Label;
+	}
+	else
+	{
+		TheSection = &InMenu->FindOrAddSection(NAME_None);
+	}
+
+	for (const TSharedRef<FWorkspaceItem>& Child : InMenuStructure->GetChildItems())
+	{
+		// Leaf nodes have valid spawner entries.
+		const TSharedPtr<FTabSpawnerEntry> Spawner = Child->AsSpawnerEntry();
+		if (Spawner.IsValid())
+		{
+			// Only show non-hidden items that have a valid spawner.
+			if (AllSpawners->Contains(Spawner.ToSharedRef()))
+			{
+				MakeSpawnerToolMenuEntry(TabManager, *TheSection, Spawner);
+			}
+		}
+		else if (Child->HasChildrenIn(*AllSpawners))
+		{
+			// Reduce the depth of the menu structure. Create a section for every odd group and a submenu for every even.
+			if (RecursionLevel % 2 == 0)
+			{
+				// Create a named section in the current menu and add all children there.
+				PopulateTabSpawnerToolMenu_Helper(InMenu, TabManager, Child, AllSpawners, RecursionLevel+1, true);
+			}
+			else
+			{
+				// Create a submenu and add all children there.
+				FToolMenuEntry& SubMenu = TheSection->AddSubMenu(
+					NAME_None,
+					Child->GetDisplayName(),
+					Child->GetTooltipText(),
+					FNewToolMenuDelegate::CreateStatic(&PopulateTabSpawnerToolMenu_Helper, TabManager, Child, AllSpawners, RecursionLevel+1, false),
+					FToolUIActionChoice(),
+					EUserInterfaceActionType::Button
+				);
+			}
+		}
+	}
+}
+
+void PopulateTabSpawnerToolMenu(const TSharedPtr<FTabManager>& TabManager, UToolMenu* InMenu, TSharedRef<FWorkspaceItem> MenuStructure, bool bIncludeOrphanedMenus)
+{
+	TSharedRef< TArray< TWeakPtr<FTabSpawnerEntry> > > AllSpawners = MakeShared< TArray< TWeakPtr<FTabSpawnerEntry> > >(TabManager->CollectSpawners());
+
+	if (bIncludeOrphanedMenus)
+	{
+		FToolMenuSection& UnnamedSection = InMenu->FindOrAddSection(NAME_None);
+
+		// Put all orphaned spawners at the top of the menu so programmers go and find them a nice home.
+		for (const TWeakPtr<FTabSpawnerEntry>& WeakSpawner : *AllSpawners)
+		{
+			const TSharedPtr<FTabSpawnerEntry> Spawner = WeakSpawner.Pin();
+			if (!Spawner)
+			{
+				continue;
+			}
+			
+			const bool bHasNoPlaceInMenuStructure = !Spawner->GetParent().IsValid();
+			if ( bHasNoPlaceInMenuStructure )
+			{
+				MakeSpawnerToolMenuEntry(TabManager, UnnamedSection, Spawner);
+			}
+		}
+	}
+	
+	PopulateTabSpawnerToolMenu_Helper(InMenu, TabManager, MenuStructure, AllSpawners, 0, false);
+}
+
+void PopulateTabSpawnerToolSection(TSharedPtr<FTabManager> TabManager, FToolMenuSection& InSection, const FName& TabType)
+{
+	TSharedPtr<FTabSpawnerEntry> Spawner = TabManager->FindTabSpawnerFor(TabType);
+	if (Spawner.IsValid())
+	{
+		MakeSpawnerToolMenuEntry(TabManager, InSection, Spawner);
+	}
+	else
+	{
+		UE_LOG(LogMainFrame, Warning, TEXT("PopulateTabSpawnerMenu failed to find entry for %s"), *(TabType.ToString()));
+	}
+}
+
+} // namespace UE::MainMenu::Private
 
 void FMainMenu::RegisterFileMenu()
 {
@@ -168,12 +280,12 @@ void FMainMenu::RegisterEditMenu()
 			);
 		}
 
-		Section.AddDynamicEntry("PluginsEditor", FNewToolMenuDelegateLegacy::CreateLambda([](FMenuBuilder& InBuilder, UToolMenu* InData)
+		Section.AddDynamicEntry("PluginsEditor", FNewToolMenuSectionDelegate::CreateLambda([](FToolMenuSection& InSection)
 		{
 			//@todo The tab system needs to be able to be extendable by plugins [9/3/2013 Justin.Sargent]
 			if (IModularFeatures::Get().IsModularFeatureAvailable(EditorFeatures::PluginsEditor))
 			{
-				FGlobalTabmanager::Get()->PopulateTabSpawnerMenu(InBuilder, "PluginsEditor");
+				UE::MainMenu::Private::PopulateTabSpawnerToolSection(FGlobalTabmanager::Get().ToSharedPtr(), InSection, FName(TEXT("PluginsEditor")));
 			}
 		}));
 	}
@@ -189,11 +301,10 @@ void FMainMenu::RegisterWindowMenu()
 
 	// Level Editor, General, and Testing sections
 	// Automatically populate tab spawners from TabManager
-	Menu->AddDynamicSection("TabManagerSection", FNewToolMenuDelegateLegacy::CreateLambda([](FMenuBuilder& InBuilder, UToolMenu* InData)
+	Menu->AddDynamicSection("TabManagerSection", FNewToolMenuDelegate::CreateLambda([](UToolMenu* InMenu)
 	{
-		if (USlateTabManagerContext* TabManagerContext = InData->FindContext<USlateTabManagerContext>())
+		if (USlateTabManagerContext* TabManagerContext = InMenu->FindContext<USlateTabManagerContext>())
 		{
-			
 			if (TSharedPtr<FTabManager> TabManager = TabManagerContext->TabManager.Pin())
 			{
 				// The global tab manager will be the tab manager for nomad tabs that appear docked as major tabs. However major tabs are not spawned
@@ -201,11 +312,11 @@ void FMainMenu::RegisterWindowMenu()
 				if(TabManager != FGlobalTabmanager::Get())
 				{
 					// Local editor tabs
-					TabManager->PopulateLocalTabSpawnerMenu(InBuilder);
+					UE::MainMenu::Private::PopulateTabSpawnerToolMenu(TabManager, InMenu, TabManager->GetLocalWorkspaceMenuRoot(), true);
 
 					// General tabs
 					const IWorkspaceMenuStructure& MenuStructure = WorkspaceMenu::GetMenuStructure();
-					TabManager->PopulateTabSpawnerMenu(InBuilder, MenuStructure.GetStructureRoot());
+					UE::MainMenu::Private::PopulateTabSpawnerToolMenu(TabManager, InMenu, MenuStructure.GetStructureRoot(), true);
 				}
 			}
 		}
@@ -502,21 +613,20 @@ void FMainMenu::RegisterToolsMenu()
 
 	// Level Editor, General, and Testing sections
 	// Automatically populate tab spawners from TabManager
-	Menu->AddDynamicSection("TabManagerSection", FNewToolMenuDelegateLegacy::CreateLambda([](FMenuBuilder& InBuilder, UToolMenu* InData)
+	Menu->AddDynamicSection("TabManagerSection", FNewToolMenuDelegate::CreateLambda([](UToolMenu* InMenu)
 	{
-		if (USlateTabManagerContext* TabManagerContext = InData->FindContext<USlateTabManagerContext>())
+		if (USlateTabManagerContext* TabManagerContext = InMenu->FindContext<USlateTabManagerContext>())
 		{
 			TSharedPtr<FTabManager> TabManager = TabManagerContext->TabManager.Pin();
 			if (TabManager.IsValid())
 			{
 				// General tabs
 				const IWorkspaceMenuStructure& MenuStructure = WorkspaceMenu::GetMenuStructure();
-				TabManager->PopulateTabSpawnerMenu(InBuilder, MenuStructure.GetToolsStructureRoot());
+				UE::MainMenu::Private::PopulateTabSpawnerToolMenu(TabManager, InMenu, MenuStructure.GetToolsStructureRoot(), true);
 			}
 		}
 	}));
 
-
 	// Experimental section
 	{
 		// This is a temporary home for the spawners of experimental features that must be explicitly enabled.
diff --git a/Engine/Source/Runtime/Slate/Private/Framework/Docking/TabManager.cpp b/Engine/Source/Runtime/Slate/Private/Framework/Docking/TabManager.cpp
--- a/Engine/Source/Runtime/Slate/Private/Framework/Docking/TabManager.cpp
+++ b/Engine/Source/Runtime/Slate/Private/Framework/Docking/TabManager.cpp
@@ -1213,9 +1213,9 @@ void FTabManager::PopulateTabSpawnerMenu(FMenuBuilder& PopulateMe, TSharedRef<FW
 	PopulateTabSpawnerMenu(PopulateMe, MenuStructure, true);
 }
 
-void FTabManager::PopulateTabSpawnerMenu( FMenuBuilder& PopulateMe, TSharedRef<FWorkspaceItem> MenuStructure, bool bIncludeOrphanedMenus )
+TArray< TWeakPtr<FTabSpawnerEntry> > FTabManager::CollectSpawners()
 {
-	TSharedRef< TArray< TWeakPtr<FTabSpawnerEntry> > > AllSpawners = MakeShareable( new TArray< TWeakPtr<FTabSpawnerEntry> >() );
+	TArray< TWeakPtr<FTabSpawnerEntry> > AllSpawners;
 
 	// Editor-specific tabs
 	for ( FTabSpawner::TIterator SpawnerIterator(TabSpawner); SpawnerIterator; ++SpawnerIterator )
@@ -1225,7 +1225,7 @@ void FTabManager::PopulateTabSpawnerMenu( FMenuBuilder& PopulateMe, TSharedRef<F
 		{
 			if (IsAllowedTab(SpawnerEntry->TabType))
 			{
-				AllSpawners->AddUnique(SpawnerEntry);
+				AllSpawners.AddUnique(SpawnerEntry);
 			}
 		}
 	}
@@ -1238,17 +1238,29 @@ void FTabManager::PopulateTabSpawnerMenu( FMenuBuilder& PopulateMe, TSharedRef<F
 		{
 			if (IsAllowedTab(SpawnerEntry->TabType))
 			{
-				AllSpawners->AddUnique(SpawnerEntry);
+				AllSpawners.AddUnique(SpawnerEntry);
 			}
 		}
 	}
 
+	return AllSpawners;
+}
+
+void FTabManager::PopulateTabSpawnerMenu( FMenuBuilder& PopulateMe, TSharedRef<FWorkspaceItem> MenuStructure, bool bIncludeOrphanedMenus )
+{
+	TSharedRef< TArray< TWeakPtr<FTabSpawnerEntry> > > AllSpawners = MakeShared< TArray< TWeakPtr<FTabSpawnerEntry> > >(CollectSpawners());
+	
 	if ( bIncludeOrphanedMenus )
 	{
 		// Put all orphaned spawners at the top of the menu so programmers go and find them a nice home.
-		for ( int32 ChildIndex=0; ChildIndex < AllSpawners->Num(); ++ChildIndex )
+		for (const TWeakPtr<FTabSpawnerEntry>& WeakSpawner : *AllSpawners)
 		{
-			const TSharedPtr<FTabSpawnerEntry> Spawner = ( *AllSpawners )[ChildIndex].Pin();
+			const TSharedPtr<FTabSpawnerEntry> Spawner = WeakSpawner.Pin();
+			if (!Spawner)
+			{
+				continue;
+			}
+
 			const bool bHasNoPlaceInMenuStructure = !Spawner->GetParent().IsValid();
 			if ( bHasNoPlaceInMenuStructure )
 			{
diff --git a/Engine/Source/Runtime/Slate/Public/Framework/Docking/TabManager.h b/Engine/Source/Runtime/Slate/Public/Framework/Docking/TabManager.h
--- a/Engine/Source/Runtime/Slate/Public/Framework/Docking/TabManager.h
+++ b/Engine/Source/Runtime/Slate/Public/Framework/Docking/TabManager.h
@@ -292,6 +292,11 @@ struct FTabSpawnerEntry : public FWorkspaceItem
 		return TabType;
 	};
 
+	const bool IsHidden() const
+	{
+		return MenuType.Get() == ETabSpawnerMenuType::Hidden;
+	}
+
 private:
 	FName TabType;
 	FOnSpawnTab OnSpawnTab;
@@ -989,6 +994,15 @@ class SLATE_API FTabManager : public TSharedFromThis<FTabManager>
 
 		FUIAction GetUIActionForTabSpawnerMenuEntry(TSharedPtr<FTabSpawnerEntry> InTabMenuEntry);
 
+		/** Return an array of all tab spawners that should have menu items in the main menu. */
+		TArray< TWeakPtr<FTabSpawnerEntry> > CollectSpawners();
+
+		/** Find a tab spawner by name. */
+		TSharedPtr<FTabSpawnerEntry> FindTabSpawnerFor(FName TabId);
+
+		/** Find a tab spawner by name. */
+		const TSharedPtr<const FTabSpawnerEntry> FindTabSpawnerFor(FName TabId) const;
+
 	protected:
 		void InvokeTabForMenu( FName TabId );
 
@@ -1105,8 +1119,6 @@ class SLATE_API FTabManager : public TSharedFromThis<FTabManager>
 	protected:
 		FTabSpawner TabSpawner;
 		TSharedRef<FTabSpawner> NomadTabSpawner;
-		TSharedPtr<FTabSpawnerEntry> FindTabSpawnerFor(FName TabId);
-		const TSharedPtr<const FTabSpawnerEntry> FindTabSpawnerFor(FName TabId) const;
 
 		bool HasTabSpawnerFor(FName TabId) const;
 
