Fix a bug causing icons to be missing from the Window menu

diff --git a/Engine/Source/Editor/MainFrame/Private/Menus/MainMenu.cpp b/Engine/Source/Editor/MainFrame/Private/Menus/MainMenu.cpp
--- a/Engine/Source/Editor/MainFrame/Private/Menus/MainMenu.cpp
+++ b/Engine/Source/Editor/MainFrame/Private/Menus/MainMenu.cpp
@@ -97,7 +97,9 @@ void PopulateTabSpawnerToolMenu_Helper(UToolMenu* InMenu, TSharedPtr<FTabManager
 					Child->GetTooltipText(),
 					FNewToolMenuDelegate::CreateStatic(&PopulateTabSpawnerToolMenu_Helper, TabManager, Child, AllSpawners, RecursionLevel+1, false),
 					FToolUIActionChoice(),
-					EUserInterfaceActionType::Button
+					EUserInterfaceActionType::Button,
+					false,
+					Child->GetIcon()
 				);
 			}
 		}
