Register FGlobalEditorCommonCommands before trying to create the main menu
#rnx

diff --git a/Engine/Source/Editor/MainFrame/Private/MainFrameModule.cpp b/Engine/Source/Editor/MainFrame/Private/MainFrameModule.cpp
--- a/Engine/Source/Editor/MainFrame/Private/MainFrameModule.cpp
+++ b/Engine/Source/Editor/MainFrame/Private/MainFrameModule.cpp
@@ -382,6 +382,9 @@ bool FMainFrameModule::IsRecreatingDefaultMainFrame() const
 
 TSharedRef<SWidget> FMainFrameModule::MakeMainMenu(const TSharedPtr<FTabManager>& TabManager, const FName MenuName, FToolMenuContext& ToolMenuContext) const
 {
+	// Can't make the MainMenu without the global editor commands having been registered
+	FGlobalEditorCommonCommands::Register();
+
 	return FMainMenu::MakeMainMenu(TabManager, MenuName, ToolMenuContext);
 }
 
