Clear the auto-save restore file after closing the main frame (and discarding changes)

This avoids the auto-save restore prompt appearing if the editor crashes after the user has already opted to discard their pending changes during shutdown

diff --git a/Engine/Source/Editor/MainFrame/Private/Frame/MainFrameHandler.cpp b/Engine/Source/Editor/MainFrame/Private/Frame/MainFrameHandler.cpp
--- a/Engine/Source/Editor/MainFrame/Private/Frame/MainFrameHandler.cpp
+++ b/Engine/Source/Editor/MainFrame/Private/Frame/MainFrameHandler.cpp
@@ -78,8 +78,11 @@ void FMainFrameHandler::ShutDownEditor()
 {
 	FEditorDelegates::OnShutdownPostPackagesSaved.Broadcast();
 
+	// By this point we've opted to discard anything we didn't want to save, so disable the auto-save restore.
+	GUnrealEd->GetPackageAutoSaver().UpdateRestoreFile(false);
 	// Any pending autosaves should not happen.  A tick will go by before the editor shuts down and we want to avoid auto-saving during this time.
 	GUnrealEd->GetPackageAutoSaver().ResetAutoSaveTimer();
+
 	GEditor->RequestEndPlayMap();
 
 	// End any play on console/PC games still happening
