Remove unused return value from FMainMenu::MakeMainMenu

This paves the way for a flickering fix of the main menus when widgets are shared between widgets where I will move the responsibility for creating widgets into the tab manager. In the tab manager, the main menu widgets will be created per window behind the scenes. Therefore, we cannot return a widget from these methods because there might be more than one and they haven't been created yet.

* Private API changes: Change the return type from TSharedRef/Ptr<SWidget> to void on various private API MakeMainMenu methods (of FLevelEditorMenu, FMainFrameModule, FMainMenu, and SStandaloneAssetEditorToolkitHost).
* Public API change: Add new method void IMainFrameModule::MakeMainMenuNoWidget.
* Public API change: Deprecate the method TSharedRef<SWidget> IMainFrameModule::MakeMainMenu and provide a default implementation for the time being that calls MakeMainMenuNoWidget.

diff --git a/Engine/Plugins/Runtime/nDisplay/Source/DisplayClusterOperator/Private/SDisplayClusterOperatorPanel.cpp b/Engine/Plugins/Runtime/nDisplay/Source/DisplayClusterOperator/Private/SDisplayClusterOperatorPanel.cpp
--- a/Engine/Plugins/Runtime/nDisplay/Source/DisplayClusterOperator/Private/SDisplayClusterOperatorPanel.cpp
+++ b/Engine/Plugins/Runtime/nDisplay/Source/DisplayClusterOperator/Private/SDisplayClusterOperatorPanel.cpp
@@ -188,7 +188,7 @@ void SDisplayClusterOperatorPanel::Construct(const FArguments& InArgs, const TSh
 	FToolMenuContext ToolMenuContext(CommandList, MenuExtender);
 	
 	const IMainFrameModule& MainFrameModule = FModuleManager::LoadModuleChecked<IMainFrameModule>("MainFrame");
-	MainFrameModule.MakeMainMenu(TabManager, MenuName, ToolMenuContext);
+	MainFrameModule.MakeMainMenuNoWidget(TabManager, MenuName, ToolMenuContext);
 	
 	// Allow any external modules to register extensions to the operator panel's status bar
 	FDisplayClusterOperatorStatusBarExtender StatusBarExtender;
diff --git a/Engine/Source/Editor/LevelEditor/Private/LevelEditorMenu.cpp b/Engine/Source/Editor/LevelEditor/Private/LevelEditorMenu.cpp
--- a/Engine/Source/Editor/LevelEditor/Private/LevelEditorMenu.cpp
+++ b/Engine/Source/Editor/LevelEditor/Private/LevelEditorMenu.cpp
@@ -325,16 +325,14 @@ void FLevelEditorMenu::RegisterLevelEditorMenus()
 	RegisterSelectMenu();
 }
 
-TSharedRef< SWidget > FLevelEditorMenu::MakeLevelEditorMenu( const TSharedPtr<FUICommandList>& CommandList, TSharedPtr<class SLevelEditor> LevelEditor )
+void FLevelEditorMenu::MakeLevelEditorMenu( const TSharedPtr<FUICommandList>& CommandList, TSharedPtr<class SLevelEditor> LevelEditor )
 {
 	FLevelEditorModule& LevelEditorModule = FModuleManager::LoadModuleChecked<FLevelEditorModule>("LevelEditor");
 	TSharedPtr<FExtender> Extenders = LevelEditorModule.GetMenuExtensibilityManager()->GetAllExtenders();
 	FToolMenuContext ToolMenuContext(CommandList, Extenders.ToSharedRef());
 
 	IMainFrameModule& MainFrameModule = FModuleManager::LoadModuleChecked<IMainFrameModule>( "MainFrame" );
-	TSharedRef< SWidget > MenuBarWidget = MainFrameModule.MakeMainMenu( LevelEditor->GetTabManager(), "LevelEditor.MainMenu", ToolMenuContext );
-
-	return MenuBarWidget;
+	MainFrameModule.MakeMainMenuNoWidget( LevelEditor->GetTabManager(), "LevelEditor.MainMenu", ToolMenuContext );
 }
 
 void FLevelEditorMenu::RegisterBuildMenu()
diff --git a/Engine/Source/Editor/LevelEditor/Private/LevelEditorMenu.h b/Engine/Source/Editor/LevelEditor/Private/LevelEditorMenu.h
--- a/Engine/Source/Editor/LevelEditor/Private/LevelEditorMenu.h
+++ b/Engine/Source/Editor/LevelEditor/Private/LevelEditorMenu.h
@@ -15,15 +15,13 @@ class FLevelEditorMenu
 {
 
 public:
-
 	static void RegisterLevelEditorMenus();
 
 	/**
-	 * Static: Creates a widget for the level editor's menu
-	 *
-	 * @return	New widget
+	 * Static: Creates a main menu for the given level editor's tab manager.
 	 */
-	static TSharedRef< SWidget > MakeLevelEditorMenu( const TSharedPtr<FUICommandList>& CommandList, TSharedPtr<class SLevelEditor> LevelEditor );
+	static void MakeLevelEditorMenu(const TSharedPtr<FUICommandList>& CommandList, TSharedPtr<class SLevelEditor> LevelEditor);
+
 private:
 	static void RegisterBuildMenu();
 	static void RegisterSelectMenu();
diff --git a/Engine/Source/Editor/LocalizationDashboard/Private/LocalizationDashboard.cpp b/Engine/Source/Editor/LocalizationDashboard/Private/LocalizationDashboard.cpp
--- a/Engine/Source/Editor/LocalizationDashboard/Private/LocalizationDashboard.cpp
+++ b/Engine/Source/Editor/LocalizationDashboard/Private/LocalizationDashboard.cpp
@@ -190,7 +190,7 @@ void SLocalizationDashboard::Construct(const FArguments& InArgs, const TSharedPt
 		]
 	];
 
-	MainFrameModule.MakeMainMenu(TabManager, "LocalizationDashboard.MainMenu", ToolMenuContext);
+	MainFrameModule.MakeMainMenuNoWidget(TabManager, "LocalizationDashboard.MainMenu", ToolMenuContext);
 
 	// Open the first game target if present
 	if (ULocalizationTargetSet* GameTargetSet = ULocalizationSettings::GetGameTargetSet())
diff --git a/Engine/Source/Editor/MainFrame/Private/MainFrameModule.cpp b/Engine/Source/Editor/MainFrame/Private/MainFrameModule.cpp
--- a/Engine/Source/Editor/MainFrame/Private/MainFrameModule.cpp
+++ b/Engine/Source/Editor/MainFrame/Private/MainFrameModule.cpp
@@ -311,7 +311,7 @@ void FMainFrameModule::CreateDefaultMainFrameAuxiliary(const bool bStartImmersiv
 
 
 			FToolMenuContext EmptyContext;
-			MakeMainMenu(FGlobalTabmanager::Get(), "MainFrame.NomadMainMenu", EmptyContext);
+			MakeMainMenuNoWidget(FGlobalTabmanager::Get(), "MainFrame.NomadMainMenu", EmptyContext);
 		
 			MainFrameContent = FGlobalTabmanager::Get()->RestoreFrom(LoadedLayout, RootWindow, WindowConfig.bEmbedTitleAreaContent, OutputCanBeNullptr);
 			// MainFrameContent will only be nullptr if its main area contains invalid tabs (probably some layout bug). If so, reset layout to avoid potential crashes
@@ -380,12 +380,12 @@ bool FMainFrameModule::IsRecreatingDefaultMainFrame() const
 }
 
 
-TSharedRef<SWidget> FMainFrameModule::MakeMainMenu(const TSharedPtr<FTabManager>& TabManager, const FName MenuName, FToolMenuContext& ToolMenuContext) const
+void FMainFrameModule::MakeMainMenuNoWidget(const TSharedPtr<FTabManager>& TabManager, const FName MenuName, FToolMenuContext& ToolMenuContext) const
 {
 	// Can't make the MainMenu without the global editor commands having been registered
 	FGlobalEditorCommonCommands::Register();
 
-	return FMainMenu::MakeMainMenu(TabManager, MenuName, ToolMenuContext);
+	FMainMenu::MakeMainMenu(TabManager, MenuName, ToolMenuContext);
 }
 
 TSharedRef<SWidget> FMainFrameModule::MakeDeveloperTools( const TArray<FMainFrameDeveloperTool>& AdditionalTools ) const
diff --git a/Engine/Source/Editor/MainFrame/Private/MainFrameModule.h b/Engine/Source/Editor/MainFrame/Private/MainFrameModule.h
--- a/Engine/Source/Editor/MainFrame/Private/MainFrameModule.h
+++ b/Engine/Source/Editor/MainFrame/Private/MainFrameModule.h
@@ -51,7 +51,7 @@ private:
 
 public:
 	virtual bool IsRecreatingDefaultMainFrame() const override;
-	virtual TSharedRef<SWidget> MakeMainMenu(const TSharedPtr<FTabManager>& TabManager, const FName MenuName, FToolMenuContext& ToolMenuContext) const override;
+	virtual void MakeMainMenuNoWidget(const TSharedPtr<FTabManager>& TabManager, const FName MenuName, FToolMenuContext& ToolMenuContext) const override;
 	
 
 	virtual TSharedRef<SWidget> MakeDeveloperTools( const TArray<FMainFrameDeveloperTool>& AdditionalTools ) const override;
diff --git a/Engine/Source/Editor/MainFrame/Private/Menus/MainMenu.cpp b/Engine/Source/Editor/MainFrame/Private/Menus/MainMenu.cpp
--- a/Engine/Source/Editor/MainFrame/Private/Menus/MainMenu.cpp
+++ b/Engine/Source/Editor/MainFrame/Private/Menus/MainMenu.cpp
@@ -425,7 +425,7 @@ void FMainMenu::RegisterHelpMenu()
 
 #undef LOCTEXT_NAMESPACE
 
-TSharedRef<SWidget> FMainMenu::MakeMainMenu(const TSharedPtr<FTabManager>& TabManager, const FName MenuName, FToolMenuContext& ToolMenuContext)
+void FMainMenu::MakeMainMenu(const TSharedPtr<FTabManager>& TabManager, const FName MenuName, FToolMenuContext& ToolMenuContext)
 {
 	// Cache all project names once
 	FMainFrameActionCallbacks::CacheProjectNames();
@@ -448,8 +448,6 @@ TSharedRef<SWidget> FMainMenu::MakeMainMenu(const TSharedPtr<FTabManager>& TabMa
 		TSharedRef<SMultiBoxWidget> MultiBoxWidget = StaticCastSharedRef<SMultiBoxWidget>(MenuBarWidget);
 		TabManager->SetMenuMultiBox(ConstCastSharedRef<FMultiBox>(MultiBoxWidget->GetMultiBox()), MultiBoxWidget);
 	}
-
-	return MenuBarWidget;
 }
 
 void FMainMenu::RegisterMainMenu()
diff --git a/Engine/Source/Editor/MainFrame/Private/Menus/MainMenu.h b/Engine/Source/Editor/MainFrame/Private/Menus/MainMenu.h
--- a/Engine/Source/Editor/MainFrame/Private/Menus/MainMenu.h
+++ b/Engine/Source/Editor/MainFrame/Private/Menus/MainMenu.h
@@ -24,9 +24,9 @@ public:
 	 *
 	 * @param TabManager Create the workspace menu based on this tab manager.
 	 * @param MenuName Identifier associated with the menu.
-	 * @return ToolMenuContext Context containing state.
+	 * @param ToolMenuContext Context containing state.
 	 */
-	static TSharedRef<SWidget> MakeMainMenu( const TSharedPtr<FTabManager>& TabManager, const FName MenuName, FToolMenuContext& ToolMenuContext );
+	static void MakeMainMenu( const TSharedPtr<FTabManager>& TabManager, const FName MenuName, FToolMenuContext& ToolMenuContext );
 
 	/**
 	 * Static: Registers main menu with menu system.
diff --git a/Engine/Source/Editor/MainFrame/Public/Interfaces/IMainFrameModule.h b/Engine/Source/Editor/MainFrame/Public/Interfaces/IMainFrameModule.h
--- a/Engine/Source/Editor/MainFrame/Public/Interfaces/IMainFrameModule.h
+++ b/Engine/Source/Editor/MainFrame/Public/Interfaces/IMainFrameModule.h
@@ -6,6 +6,7 @@
 #include "SlateFwd.h"
 #include "Modules/ModuleInterface.h"
 #include "Modules/ModuleManager.h"
+#include "Widgets/SNullWidget.h"
 #include "Widgets/SWidget.h"
 #include "Framework/Commands/UICommandList.h"
 #include "Framework/MultiBox/MultiBoxExtender.h"
@@ -62,12 +63,27 @@ public:
 	 * a menu for your tab, you should call this function to create your menu, passing in an extender object to add your
 	 * tab-specific menu items!
 	 *
-	 * @param	TabManager	The tab manager for the tab you're creating the menu for.  This is needed so we can populate the layout menus correctly.
+	 * @param	TabManager	The tab manager for the tab you're creating the menu for. This is needed so we can populate the layout menus correctly.
 	 * @param	Extender	Extender object used to customize the main frame menu
 	 *
 	 * @return	The newly-created menu widget
 	 */
-	virtual TSharedRef<SWidget> MakeMainMenu( const TSharedPtr<FTabManager>& TabManager, const FName MenuName, FToolMenuContext& ToolMenuContext ) const = 0;
+	UE_DEPRECATED(5.4, "Generation of main menu widgets is moving into the tab manager and therefore we cannot reliably return a widget here. Please call MakeMainMenuNoWidget instead.")
+	virtual TSharedRef<SWidget> MakeMainMenu(const TSharedPtr<FTabManager>& TabManager, const FName MenuName, FToolMenuContext& ToolMenuContext) const
+	{
+		MakeMainMenuNoWidget(TabManager, MenuName, ToolMenuContext);
+		return SNullWidget::NullWidget;
+	}
+
+	/**
+	 * Generates a menu that includes application global commands, such as "Save All", "Exit", etc.  If you're building
+	 * a menu for your tab, you should call this function to create your menu, passing in an extender object to add your
+	 * tab-specific menu items!
+	 *
+	 * @param	TabManager	The tab manager for the tab you're creating the menu for. This is needed so we can populate the layout menus correctly.
+	 * @param	Extender	Extender object used to customize the main frame menu
+	 */
+	virtual void MakeMainMenuNoWidget( const TSharedPtr<FTabManager>& TabManager, const FName MenuName, FToolMenuContext& ToolMenuContext ) const = 0;
 
 	/**
 	 * Generates a menu that's just like the "main menu" widget above, except it also includes some infrequently used commands
diff --git a/Engine/Source/Editor/UnrealEd/Private/Toolkits/SStandaloneAssetEditorToolkitHost.cpp b/Engine/Source/Editor/UnrealEd/Private/Toolkits/SStandaloneAssetEditorToolkitHost.cpp
--- a/Engine/Source/Editor/UnrealEd/Private/Toolkits/SStandaloneAssetEditorToolkitHost.cpp
+++ b/Engine/Source/Editor/UnrealEd/Private/Toolkits/SStandaloneAssetEditorToolkitHost.cpp
@@ -69,7 +69,7 @@ void SStandaloneAssetEditorToolkitHost::SetupInitialContent( const TSharedRef<FT
 		}
 	}
 
-	DefaultMenuWidget = SNullWidget::NullWidget;
+	bMainMenuCreated = false;
 
 	HostTabPtr = InHostTab;
 
@@ -231,7 +231,7 @@ FName SStandaloneAssetEditorToolkitHost::GetMenuName() const
 
 void SStandaloneAssetEditorToolkitHost::GenerateMenus(bool bForceCreateMenu)
 {
-	if( bForceCreateMenu || DefaultMenuWidget != SNullWidget::NullWidget )
+	if( bForceCreateMenu || bMainMenuCreated )
 	{
 		const FName AssetEditorMenuName = GetMenuName();
 
@@ -240,7 +240,9 @@ void SStandaloneAssetEditorToolkitHost::GenerateMenus(bool bForceCreateMenu)
 		FToolMenuContext ToolMenuContext(HostedAssetEditorToolkit->GetToolkitCommands(), FExtender::Combine(MenuExtenders), ContextObject);
 		HostedAssetEditorToolkit->InitToolMenuContext(ToolMenuContext);
 		IMainFrameModule& MainFrameModule = FModuleManager::LoadModuleChecked<IMainFrameModule>( "MainFrame" );
-		DefaultMenuWidget = MainFrameModule.MakeMainMenu( MyTabManager, AssetEditorMenuName, ToolMenuContext );
+
+		MainFrameModule.MakeMainMenuNoWidget( MyTabManager, AssetEditorMenuName, ToolMenuContext );
+		bMainMenuCreated = true;
 	}
 }
 
diff --git a/Engine/Source/Editor/UnrealEd/Private/Toolkits/SStandaloneAssetEditorToolkitHost.h b/Engine/Source/Editor/UnrealEd/Private/Toolkits/SStandaloneAssetEditorToolkitHost.h
--- a/Engine/Source/Editor/UnrealEd/Private/Toolkits/SStandaloneAssetEditorToolkitHost.h
+++ b/Engine/Source/Editor/UnrealEd/Private/Toolkits/SStandaloneAssetEditorToolkitHost.h
@@ -124,8 +124,8 @@ private:
 	/** The widget that will house the overlay widgets (if any) */
 	TSharedPtr<SBox> MenuOverlayWidgetContent;
 
-	/** The default menu widget */
-	TSharedPtr< SWidget > DefaultMenuWidget;
+	/** Track whether we've created a main menu yet. */
+	bool bMainMenuCreated = false;
 
 	/** The DockTab in which we reside. */
 	TWeakPtr<SDockTab> HostTabPtr;
