Add cvar to allow re-opening asset editors when -PIE is present

#tests Launched with PIE and a bunch of asset windows open from last run

diff --git a/Engine/Source/Editor/MainFrame/Private/Frame/MainFrameHandler.cpp b/Engine/Source/Editor/MainFrame/Private/Frame/MainFrameHandler.cpp
--- a/Engine/Source/Editor/MainFrame/Private/Frame/MainFrameHandler.cpp
+++ b/Engine/Source/Editor/MainFrame/Private/Frame/MainFrameHandler.cpp
@@ -12,6 +12,10 @@
 static bool ShowRestoreAssetsPromptOnStartup = true;
 FAutoConsoleVariableRef CVarShowRestoreAssetsPromptOnStartup(TEXT("Mainframe.ShowRestoreAssetsPromptOnStartup"), ShowRestoreAssetsPromptOnStartup, TEXT(""), ECVF_ReadOnly);
 
+static bool ShowRestoreAssetsPromptInPIE = false;
+FAutoConsoleVariableRef CVarShowRestoreAssetsPromptInPIE(TEXT("Mainframe.ShowRestoreAssetsPromptInPIE"), ShowRestoreAssetsPromptInPIE,
+	TEXT("Restore asset windows when running with PIE at startup (default: false). This doesn't work with immersive mode or if Mainframe.ShowRestoreAssetsPromptOnStartup is set to false."));
+
 void FMainFrameHandler::ShowMainFrameWindow(TSharedRef<SWindow> Window, const bool bStartImmersive, const bool bStartPIE) const
 {
 	// Make sure viewport windows are maximized/immersed if they need to be
@@ -39,6 +43,18 @@ void FMainFrameHandler::ShowMainFrameWindow(TSharedRef<SWindow> Window, const bo
 		// Need to register after the window is shown or else we cant capture the mouse
 		TSharedPtr<IAssetViewport> Viewport = LevelEditor.GetFirstActiveViewport();
 		Viewport->RegisterGameViewportIfPIE();
+
+		if (!bStartImmersive)
+		{
+			// Command lines may not have processed yet, flush them here to make sure cvars are up to date.
+			GEngine->TickDeferredCommands();
+
+			if (ShowRestoreAssetsPromptInPIE && ShowRestoreAssetsPromptOnStartup)
+			{
+				// Restore any assets we had open. Can be requested to function in PIE
+				GEditor->GetEditorSubsystem<UAssetEditorSubsystem>()->RequestRestorePreviouslyOpenAssets();
+			}
+		}
 	}
 	else
 	{
