Fix an issue where the OnTabClosed delegate is not triggerd when the Tabs are closed during the Mainframe reinitialization process.

The issue was that we were destroying all windows before giving a change to the tab manager to properly close them.

diff --git a/Engine/Source/Editor/MainFrame/Private/Frame/MainFrameHandler.h b/Engine/Source/Editor/MainFrame/Private/Frame/MainFrameHandler.h
--- a/Engine/Source/Editor/MainFrame/Private/Frame/MainFrameHandler.h
+++ b/Engine/Source/Editor/MainFrame/Private/Frame/MainFrameHandler.h
@@ -67,6 +67,13 @@ public:
 	 */
 	bool CanCloseEditor()
 	{
+		// We don't want to close the editor if we are recreating the mainframe module
+		IMainFrameModule& MainFrameModule = FModuleManager::LoadModuleChecked<IMainFrameModule>(TEXT("MainFrame"));
+		if (MainFrameModule.IsRecreatingDefaultMainFrame())
+		{
+			return false;
+		}
+
 		if ( FSlateApplication::IsInitialized() && !FSlateApplication::Get().IsNormalExecution())
 		{
 			// DEBUGGER EXIT PATH
@@ -136,7 +143,6 @@ public:
 			}
 			
 			// Allow Plugins and other systems to prevent close 
-			IMainFrameModule& MainFrameModule = FModuleManager::LoadModuleChecked<IMainFrameModule>(TEXT("MainFrame"));
 			bOkToExit = bOkToExit && MainFrameModule.ExecuteCanCloseEditorDelegates();
 
 			// Prompt for save and quit only if we did not launch a gameless rocket exe or are in demo mode or we are asking for a close to recreate the Default Main Frame
diff --git a/Engine/Source/Editor/MainFrame/Private/MainFrameModule.cpp b/Engine/Source/Editor/MainFrame/Private/MainFrameModule.cpp
--- a/Engine/Source/Editor/MainFrame/Private/MainFrameModule.cpp
+++ b/Engine/Source/Editor/MainFrame/Private/MainFrameModule.cpp
@@ -133,10 +133,11 @@ void FMainFrameModule::RecreateDefaultMainFrame(const bool bStartImmersive, cons
 	// Clean previous default main frame
 	if (IsWindowInitialized())
 	{
-		// Clean FSlateApplication
-		FSlateApplication::Get().CloseAllWindowsImmediately();
 		// Clean FGlobalTabmanager
 		FGlobalTabmanager::Get()->CloseAllAreas();
+
+		// Clean FSlateApplication
+		FSlateApplication::Get().CloseAllWindowsImmediately();
 	}
 	// (Re-)create default main frame
 	CreateDefaultMainFrameAuxiliary(bStartImmersive, bStartPIE, /*bIsBeingRecreated*/true);
@@ -317,9 +318,16 @@ void FMainFrameModule::CreateDefaultMainFrameAuxiliary(const bool bStartImmersiv
 			// MainFrameContent will only be nullptr if its main area contains invalid tabs (probably some layout bug). If so, reset layout to avoid potential crashes
 			if (!MainFrameContent.IsValid())
 			{
+				// This code will recurse to load the default layout. If that also fail we will get into an infinite loop
+				check(!bIsReCreatingDefaultLayoutAsFallback)
+				TGuardValue<bool> GuardRecreatingDefaultLayoutFallback(bIsReCreatingDefaultLayoutAsFallback, true);
+
+				TGuardValue<bool> GuardRecreatingDefaultMainFrame(bRecreatingDefaultMainFrame, true);
+
 				// Clean FSlateApplication & FGlobalTabmanager
-				FSlateApplication::Get().CloseAllWindowsImmediately();
 				FGlobalTabmanager::Get()->CloseAllAreas();
+				FSlateApplication::Get().CloseAllWindowsImmediately();
+
 				// Remove and reload file
 				GConfig->UnloadFile(GEditorLayoutIni); // We must re-read it to avoid the Editor to use a previously cached name and description
 				const FString FaultyEditorLayoutPath = GEditorLayoutIni + TEXT("_faulty.ini");
diff --git a/Engine/Source/Editor/MainFrame/Private/MainFrameModule.h b/Engine/Source/Editor/MainFrame/Private/MainFrameModule.h
--- a/Engine/Source/Editor/MainFrame/Private/MainFrameModule.h
+++ b/Engine/Source/Editor/MainFrame/Private/MainFrameModule.h
@@ -303,4 +303,7 @@ private:
 
 	// Instantiation of the object responsible for spawning the editor's project dialog on startup
 	FProjectDialogProvider ProjectDialogProvider;
+
+	// Set to true if we are in the middle of recreating the default layout because a previous load attempt failed
+	bool bIsReCreatingDefaultLayoutAsFallback = false;
 };
