Add an option to confirm closing the Unreal Editor

the pre flight is actually 662bc15be774eebeff65c9a2

diff --git a/Engine/Source/Editor/MainFrame/Private/Frame/MainFrameHandler.h b/Engine/Source/Editor/MainFrame/Private/Frame/MainFrameHandler.h
--- a/Engine/Source/Editor/MainFrame/Private/Frame/MainFrameHandler.h
+++ b/Engine/Source/Editor/MainFrame/Private/Frame/MainFrameHandler.h
@@ -21,8 +21,12 @@
 #include "LevelEditor.h"
 #include "IAssetViewport.h"
 #include "MainFrameLog.h"
+#include "Editor/EditorPerProjectUserSettings.h"
+#include "Misc/MessageDialog.h"
 #include "Subsystems/AssetEditorSubsystem.h"
 
+#define LOCTEXT_NAMESPACE "MainFrameHandler"
+
 const FText StaticGetApplicationTitle( const bool bIncludeGameName );
 
 /**
@@ -161,13 +165,25 @@ public:
 				}
 
 				// If there were packages to save, or switching project, then the user already had a chance to bail out of exiting.
-				if ( !bOkToExit && !bHadPackagesToSave && FUnrealEdMisc::Get().GetPendingProjectName().IsEmpty() )
+				const bool bPromptedForExit = bHadPackagesToSave || !FUnrealEdMisc::Get().GetPendingProjectName().IsEmpty();
+				if ( !bOkToExit && !bPromptedForExit )
 				{
 					FUnrealEdMisc::Get().ClearPendingProjectName();
 					FUnrealEdMisc::Get().AllowSavingLayoutOnClose(true);
 					FUnrealEdMisc::Get().ForceDeletePreferences(false);
 					FUnrealEdMisc::Get().ClearConfigRestoreFilenames();
 				}
+				else if (bOkToExit && !bPromptedForExit && GetDefault<UEditorPerProjectUserSettings>()->bConfirmEditorClose)
+				{
+					const EAppReturnType::Type Response = FMessageDialog::Open(
+						EAppMsgCategory::Info,
+						EAppMsgType::YesNo,
+						LOCTEXT("ConfirmClose", "Are you sure you want to close the Unreal Editor?"),
+						LOCTEXT("ConfirmCloseTitle", "Close Editor")
+					);
+
+					bOkToExit = (Response == EAppReturnType::Yes);
+				}
 			}
 			
 			return bOkToExit;
@@ -256,3 +272,5 @@ private:
 	/** The window that all of the editor is parented to. */
 	TWeakPtr<SWindow> RootWindowPtr;
 };
+
+#undef LOCTEXT_NAMESPACE
diff --git a/Engine/Source/Editor/UnrealEd/Classes/Editor/EditorPerProjectUserSettings.h b/Engine/Source/Editor/UnrealEd/Classes/Editor/EditorPerProjectUserSettings.h
--- a/Engine/Source/Editor/UnrealEd/Classes/Editor/EditorPerProjectUserSettings.h
+++ b/Engine/Source/Editor/UnrealEd/Classes/Editor/EditorPerProjectUserSettings.h
@@ -95,6 +95,10 @@ class UEditorPerProjectUserSettings : public UObject
 	UPROPERTY(EditAnywhere, config, Category = Import)
 	uint32 bAnimationReimportWarnings: 1;
 
+	/** If enabled, you'll be asked to confirm if you really want to close the editor. */
+	UPROPERTY(EditAnywhere, Config, Category = Behavior)
+	uint32 bConfirmEditorClose: 1;
+
 	/** Select to make Distributions use the curves, not the baked lookup tables. */
 	UPROPERTY(config)
 	uint32 bUseCurvesForDistributions:1; //(GDistributionType == 0)
