SkeinUE: Fixed 'Tools -> Sync Content', which can result in a world reload, taking the user out of their selected interactive editor mode and losing part of their progress.

diff --git a/Engine/Source/Editor/MainFrame/Private/Frame/MainFrameActions.cpp b/Engine/Source/Editor/MainFrame/Private/Frame/MainFrameActions.cpp
--- a/Engine/Source/Editor/MainFrame/Private/Frame/MainFrameActions.cpp
+++ b/Engine/Source/Editor/MainFrame/Private/Frame/MainFrameActions.cpp
@@ -109,10 +109,10 @@ void FMainFrameCommands::RegisterCommands()
 	ActionList->MapAction(ViewChangelists, FExecuteAction::CreateStatic( &FMainFrameActionCallbacks::ViewChangelists ), FCanExecuteAction::CreateStatic( &FMainFrameActionCallbacks::CanViewChangelists ) );
 
 	UI_COMMAND( SubmitContent, "Submit Content", "Opens a dialog with check in options for content and levels.", EUserInterfaceActionType::Button, FInputChord() );
-	ActionList->MapAction( SubmitContent, FExecuteAction::CreateLambda([]() { FSourceControlWindows::ChoosePackagesToCheckIn(); }), FCanExecuteAction::CreateStatic(&FSourceControlWindows::CanChoosePackagesToCheckIn ), FGetActionCheckState(), FIsActionButtonVisible::CreateStatic(FSourceControlWindows::ShouldChoosePackagesToCheckBeVisible) );
+	ActionList->MapAction( SubmitContent, FExecuteAction::CreateLambda([]() { FSourceControlWindows::ChoosePackagesToCheckIn(); }), FCanExecuteAction::CreateStatic( &FMainFrameActionCallbacks::CanSubmitContent ), FGetActionCheckState(), FIsActionButtonVisible::CreateStatic(FSourceControlWindows::ShouldChoosePackagesToCheckBeVisible) );
 
 	UI_COMMAND( SyncContent, "Sync Content", "Saves all unsaved levels and assets to disk and then downloads the latest versions from revision control.", EUserInterfaceActionType::Button, FInputChord() );
-	ActionList->MapAction(SyncContent, FExecuteAction::CreateLambda([]() { FSourceControlWindows::SyncLatest(); }), FCanExecuteAction::CreateStatic( &FSourceControlWindows::CanSyncLatest ) );
+	ActionList->MapAction( SyncContent, FExecuteAction::CreateLambda([]() {  FSourceControlWindows::SyncLatest(); }), FCanExecuteAction::CreateStatic( &FMainFrameActionCallbacks::CanSyncContent ) );
 
 	UI_COMMAND( ConnectToSourceControl, "Connect to Revision Control...", "Connect to a revision control system for tracking changes to your content and levels.", EUserInterfaceActionType::Button, FInputChord() );
 	ActionList->MapAction( ConnectToSourceControl, FExecuteAction::CreateStatic( &FMainFrameActionCallbacks::ConnectToSourceControl ), DefaultExecuteAction );
@@ -332,6 +332,16 @@ void FMainFrameActionCallbacks::SaveAll()
 	FEditorFileUtils::SaveDirtyPackages( bPromptUserToSave, bSaveMapPackages, bSaveContentPackages, bFastSave, bNotifyNoPackagesSaved, bCanBeDeclined );
 }
 
+bool FMainFrameActionCallbacks::CanSubmitContent()
+{
+	return FSourceControlWindows::CanChoosePackagesToCheckIn();
+}
+
+bool FMainFrameActionCallbacks::CanSyncContent()
+{
+	return GLevelEditorModeTools().CanAutoSave() && FSourceControlWindows::CanSyncLatest();
+}
+
 TArray<FRecentProjectFile> FMainFrameActionCallbacks::RecentProjects;
 
 void FMainFrameActionCallbacks::CacheProjectNames()
diff --git a/Engine/Source/Editor/MainFrame/Private/Frame/MainFrameActions.h b/Engine/Source/Editor/MainFrame/Private/Frame/MainFrameActions.h
--- a/Engine/Source/Editor/MainFrame/Private/Frame/MainFrameActions.h
+++ b/Engine/Source/Editor/MainFrame/Private/Frame/MainFrameActions.h
@@ -110,6 +110,12 @@ public:
 
 	/** Saves all levels and asset packages */
 	static void SaveAll();
+
+	/** Determine whether we are allowed to submit content at this moment */
+	static bool CanSubmitContent();
+
+	/** Determine whether we are allocated to sync content at this moment */
+	static bool CanSyncContent();
 	
 	/** Opens a dialog to choose packages to save */
 	static void ChoosePackagesToSave();
@@ -200,9 +206,6 @@ public:
 	/** Whether menu to add code to the current project is visible */
 	static bool IsAddCodeToProjectVisible();
 
-	/** Checks whether a menu action for packaging the project can execute. */
-	static bool PackageProjectCanExecute( const FName PlatformInfoName );
-
 	/** Refresh the project in the current IDE */
 	static void RefreshCodeProject();
 
