[Editor Layout] Suppressed warning for optional layout config fields...

diff --git a/Engine/Source/Editor/MainFrame/Private/Menus/LayoutsMenu.cpp b/Engine/Source/Editor/MainFrame/Private/Menus/LayoutsMenu.cpp
--- a/Engine/Source/Editor/MainFrame/Private/Menus/LayoutsMenu.cpp
+++ b/Engine/Source/Editor/MainFrame/Private/Menus/LayoutsMenu.cpp
@@ -368,8 +368,9 @@ void FPrivateLayoutsMenu::DisplayLayoutsInternal(FToolMenuSection& InSection, co
 		if (FLayoutSaveRestore::IsValidConfig(LayoutFilePath))
 		{
 			// Read and display localization name from INI file
-			const FText LayoutName = FLayoutSaveRestore::LoadSectionFromConfig(LayoutFilePath, "LayoutName");
-			const FText LayoutDescription = FLayoutSaveRestore::LoadSectionFromConfig(LayoutFilePath, "LayoutDescription");
+			constexpr bool bAreSectionsOptional = true;
+			const FText LayoutName = FLayoutSaveRestore::LoadSectionFromConfig(LayoutFilePath, "LayoutName", bAreSectionsOptional);
+			const FText LayoutDescription = FLayoutSaveRestore::LoadSectionFromConfig(LayoutFilePath, "LayoutDescription", bAreSectionsOptional);
 			// If no localization name, then display the file name
 			const FText DisplayName = (!LayoutName.IsEmpty() ? LayoutName : GetDisplayTextInternal(InLayoutIniFileNames[LayoutIndex]));
 			const FText Tooltip = GetTooltipTextInternal(DisplayName, LayoutFilePath, LayoutDescription);
@@ -391,7 +392,6 @@ void FPrivateLayoutsMenu::DisplayLayoutsInternal(FToolMenuSection& InSection, co
 					// Cannot save engine or project layouts
 					continue;
 				}
-
 			}
 			else if (InLayoutsMenu == ELayoutsMenu::Remove)
 			{
diff --git a/Engine/Source/Runtime/Slate/Private/Framework/Docking/LayoutService.cpp b/Engine/Source/Runtime/Slate/Private/Framework/Docking/LayoutService.cpp
--- a/Engine/Source/Runtime/Slate/Private/Framework/Docking/LayoutService.cpp
+++ b/Engine/Source/Runtime/Slate/Private/Framework/Docking/LayoutService.cpp
@@ -334,14 +334,19 @@ void FLayoutSaveRestore::SaveSectionToConfig(const FString& InConfigFileName, co
 	}
 }
 
-FText FLayoutSaveRestore::LoadSectionFromConfig(const FString& InConfigFileName, const FString& InSectionName)
+FText FLayoutSaveRestore::LoadSectionFromConfig(const FString& InConfigFileName, const FString& InSectionName, const bool bIsOptional)
 {
 	FString ValueString;
 	const FString JsonFileName = GetLayoutJsonFileName(InConfigFileName);
 	TSharedPtr<FJsonObject> JsonObject = LoadJsonFile(JsonFileName);
 	if (JsonObject.IsValid())
 	{
-		ValueString = JsonObject->GetStringField(InSectionName);
+		// If optional, we check if the field exists first
+		if (!bIsOptional
+			|| (bIsOptional && JsonObject->HasTypedField(InSectionName, EJson::String)))
+		{
+			ValueString = JsonObject->GetStringField(InSectionName);
+		}
 	}
 
 	if (ValueString.IsEmpty())
diff --git a/Engine/Source/Runtime/Slate/Public/Framework/Docking/LayoutService.h b/Engine/Source/Runtime/Slate/Public/Framework/Docking/LayoutService.h
--- a/Engine/Source/Runtime/Slate/Public/Framework/Docking/LayoutService.h
+++ b/Engine/Source/Runtime/Slate/Public/Framework/Docking/LayoutService.h
@@ -55,10 +55,11 @@ struct FLayoutSaveRestore
 	 *
 	 * @param InConfigFileName file to be used to load an existing layout.
 	 * @param InSectionName the name of the section to be read.
+	 * @param bIsOptional if false, an error may be output if the section doesn't exist.
 	 *
 	 * @return Loaded FText associated for that section.
 	 */
-	static SLATE_API FText LoadSectionFromConfig(const FString& InConfigFileName, const FString& InSectionName);
+	static SLATE_API FText LoadSectionFromConfig(const FString& InConfigFileName, const FString& InSectionName, const bool bIsOptional = false);
 
 	/**
 	 * Migrates the layout configuration from one config file to another.
