SourceControl: Align 'Tools' Revision Control menu with bottom right 'Revision Control' menu.

- 'View changes' only visible when source control provider uses (uncontrolled) changelists.
- 'Open snapshot history' only visible source control provider equals URC.
- 'Check out modified files' only visible when the source control provider uses checkouts.

'Submit Content' and 'Sync Latest' have been removed from these menu's in an earlier changelist.

#rnx

diff --git a/Engine/Source/Developer/SourceControl/Private/RevisionControlStyle/RevisionControlStyle.cpp b/Engine/Source/Developer/SourceControl/Private/RevisionControlStyle/RevisionControlStyle.cpp
--- a/Engine/Source/Developer/SourceControl/Private/RevisionControlStyle/RevisionControlStyle.cpp
+++ b/Engine/Source/Developer/SourceControl/Private/RevisionControlStyle/RevisionControlStyle.cpp
@@ -97,6 +97,7 @@ FDefaultRevisionControlStyle::FDefaultRevisionControlStyle() : FSlateStyleSet(St
 	Set("RevisionControl.Actions.MakeWritable", new IMAGE_BRUSH_SVG("Starship/Common/edit", CoreStyleConstants::Icon16x16));
 	Set("RevisionControl.Actions.Add", new IMAGE_BRUSH_SVG("Starship/SourceControl/RC_MarkedForAdd", CoreStyleConstants::Icon16x16));
 	Set("RevisionControl.Actions.Submit", new IMAGE_BRUSH_SVG("Starship/SourceControl/RC_CheckIn", CoreStyleConstants::Icon16x16));
+	Set("RevisionControl.Actions.Connect", new IMAGE_BRUSH_SVG("Starship/SourceControl/Status/RevisionControl", CoreStyleConstants::Icon16x16));
 	Set("RevisionControl.Actions.History", new IMAGE_BRUSH_SVG("Starship/Common/Recent", CoreStyleConstants::Icon16x16));
 	Set("RevisionControl.Actions.Diff", new IMAGE_BRUSH_SVG("Starship/SourceControl/RC_Diff", CoreStyleConstants::Icon16x16));
 	Set("RevisionControl.Actions.Revert", new IMAGE_BRUSH_SVG("Starship/SourceControl/icon_SCC_Revert", CoreStyleConstants::Icon16x16));
diff --git a/Engine/Source/Editor/MainFrame/Private/Frame/MainFrameActions.cpp b/Engine/Source/Editor/MainFrame/Private/Frame/MainFrameActions.cpp
--- a/Engine/Source/Editor/MainFrame/Private/Frame/MainFrameActions.cpp
+++ b/Engine/Source/Editor/MainFrame/Private/Frame/MainFrameActions.cpp
@@ -108,6 +108,12 @@ void FMainFrameCommands::RegisterCommands()
 	UI_COMMAND( ViewChangelists, "View Changes", "Opens a dialog displaying current changes.", EUserInterfaceActionType::Button, FInputChord() );
 	ActionList->MapAction(ViewChangelists, FExecuteAction::CreateStatic( &FMainFrameActionCallbacks::ViewChangelists ), FCanExecuteAction::CreateStatic( &FMainFrameActionCallbacks::CanViewChangelists ) );
 
+	UI_COMMAND( ViewSnapshotHistory, "Open Snapshot History", "See all the changes that have been made to this project over time.", EUserInterfaceActionType::Button, FInputChord());
+	ActionList->MapAction( ViewSnapshotHistory, FExecuteAction::CreateStatic( &FMainFrameActionCallbacks::ViewSnapshotHistory ), FCanExecuteAction::CreateStatic( &FMainFrameActionCallbacks::CanViewSnapshotHistory ) );
+
+	UI_COMMAND( CheckOutModifiedFiles, "Check Out Modified Files", "Opens a dialog to check out any assets which have been modified.", EUserInterfaceActionType::Button, FInputChord());
+	ActionList->MapAction( CheckOutModifiedFiles, FExecuteAction::CreateStatic( &FMainFrameActionCallbacks::CheckOutModifiedFiles ), FCanExecuteAction::CreateStatic( &FMainFrameActionCallbacks::CanCheckOutModifiedFiles ) );
+
 	UI_COMMAND( ConnectToSourceControl, "Connect to Revision Control...", "Connect to a revision control system for tracking changes to your content and levels.", EUserInterfaceActionType::Button, FInputChord() );
 	ActionList->MapAction( ConnectToSourceControl, FExecuteAction::CreateStatic( &FMainFrameActionCallbacks::ConnectToSourceControl ), DefaultExecuteAction );
 
@@ -304,6 +310,46 @@ bool FMainFrameActionCallbacks::CanViewChangelists()
 	return ISourceControlWindowsModule::Get().CanShowChangelistsTab();
 }
 
+void FMainFrameActionCallbacks::ViewSnapshotHistory()
+{
+	return ISourceControlWindowsModule::Get().ShowSnapshotHistoryTab();
+}
+
+bool FMainFrameActionCallbacks::CanViewSnapshotHistory()
+{
+	return ISourceControlWindowsModule::Get().CanShowSnapshotHistoryTab();
+}
+
+void FMainFrameActionCallbacks::CheckOutModifiedFiles()
+{
+	TArray<UPackage*> PackagesToSave;
+	FEditorFileUtils::GetDirtyWorldPackages(PackagesToSave);
+	FEditorFileUtils::GetDirtyContentPackages(PackagesToSave);
+
+	FEditorFileUtils::FPromptForCheckoutAndSaveParams SaveParams;
+	SaveParams.bCheckDirty = true;
+	SaveParams.bPromptToSave = false;
+	SaveParams.bIsExplicitSave = true;
+
+	FEditorFileUtils::PromptForCheckoutAndSave(PackagesToSave, SaveParams);
+}
+
+bool FMainFrameActionCallbacks::CanCheckOutModifiedFiles()
+{
+	ISourceControlModule& SourceControlModule = ISourceControlModule::Get();
+	if (ISourceControlModule::Get().IsEnabled() &&
+		ISourceControlModule::Get().GetProvider().IsAvailable())
+	{
+		TArray<UPackage*> PackagesToSave;
+		FEditorFileUtils::GetDirtyWorldPackages(PackagesToSave);
+		FEditorFileUtils::GetDirtyContentPackages(PackagesToSave);
+
+		return PackagesToSave.Num() > 0;
+	}
+
+	return false;
+}
+
 void FMainFrameActionCallbacks::ConnectToSourceControl()
 {
 	ELoginWindowMode::Type Mode = !FSlateApplication::Get().GetActiveModalWindow().IsValid() ? ELoginWindowMode::Modeless : ELoginWindowMode::Modal;
diff --git a/Engine/Source/Editor/MainFrame/Private/Frame/MainFrameActions.h b/Engine/Source/Editor/MainFrame/Private/Frame/MainFrameActions.h
--- a/Engine/Source/Editor/MainFrame/Private/Frame/MainFrameActions.h
+++ b/Engine/Source/Editor/MainFrame/Private/Frame/MainFrameActions.h
@@ -27,6 +27,8 @@ public:
 	TSharedPtr< FUICommandInfo > Exit;
 	TSharedPtr< FUICommandInfo > ChooseFilesToSave;
 	TSharedPtr< FUICommandInfo > ViewChangelists;
+	TSharedPtr< FUICommandInfo > ViewSnapshotHistory;
+	TSharedPtr< FUICommandInfo > CheckOutModifiedFiles;
 	TSharedPtr< FUICommandInfo > ConnectToSourceControl;
 	TSharedPtr< FUICommandInfo > ChangeSourceControlSettings;
 	TSharedPtr< FUICommandInfo > NewProject;
@@ -118,6 +120,18 @@ public:
 	/** Determines whether we can show the changelist view */
 	static bool CanViewChangelists();
 
+	/** Opens a dialog to view the past snapshots */
+	static void ViewSnapshotHistory();
+
+	/** Determines whether we can show the past snapshots view */
+	static bool CanViewSnapshotHistory();
+
+	/** Opens a dialog to check out any assets which have been modified */
+	static void CheckOutModifiedFiles();
+
+	/** Determines whether we can show the dialog to check out any assets which have been modified */
+	static bool CanCheckOutModifiedFiles();
+
 	/** Enable source control features */
 	static void ConnectToSourceControl();
 
diff --git a/Engine/Source/Editor/MainFrame/Private/Menus/MainMenu.cpp b/Engine/Source/Editor/MainFrame/Private/Menus/MainMenu.cpp
--- a/Engine/Source/Editor/MainFrame/Private/Menus/MainMenu.cpp
+++ b/Engine/Source/Editor/MainFrame/Private/Menus/MainMenu.cpp
@@ -32,6 +32,7 @@
 #include "ITranslationEditor.h"
 #include "LauncherPlatformModule.h"
 #include "MainFrameLog.h"
+#include "RevisionControlStyle/RevisionControlStyle.h"
 
 #define LOCTEXT_NAMESPACE "MainFileMenu"
 
@@ -657,26 +658,66 @@ void FMainMenu::RegisterToolsMenu()
 		}
 	}
 
+	// NOTE: This menu should be in sync with SourceControlMenuHelpers's SourceControl menu.
+
 	FToolMenuSection& SourceControlSection = Menu->AddSection("Source Control", LOCTEXT("SourceControlHeading", "Revision Control"));
 
-	SourceControlSection.AddMenuEntry(
-		FMainFrameCommands::Get().ViewChangelists,
-		TAttribute<FText>(),
-		TAttribute<FText>(),
-		FSlateIcon(FAppStyle::GetAppStyleSetName(), "SourceControl.ChangelistsTab")
-	);
+	SourceControlSection.AddDynamicEntry("ChangelistsTab", FNewToolMenuSectionDelegate::CreateLambda([](FToolMenuSection& InSection)
+	{
+		ISourceControlModule& SourceControlModule = ISourceControlModule::Get();
+
+		if (SourceControlModule.GetProvider().UsesChangelists() || SourceControlModule.GetProvider().UsesUncontrolledChangelists())
+		{
+			InSection.AddMenuEntry(
+				FMainFrameCommands::Get().ViewChangelists,
+				TAttribute<FText>(),
+				TAttribute<FText>(),
+				FSlateIcon(FRevisionControlStyleManager::Get().GetStyleSetName(), "RevisionControl.ChangelistsTab")
+			);
+		}
+	}));
+
+	SourceControlSection.AddDynamicEntry("SnapshotHistoryTab", FNewToolMenuSectionDelegate::CreateLambda([](FToolMenuSection& InSection)
+	{
+		ISourceControlModule& SourceControlModule = ISourceControlModule::Get();
+
+		if (SourceControlModule.GetProvider().GetName() == TEXT("Unreal Revision Control"))
+		{
+			InSection.AddMenuEntry(
+				FMainFrameCommands::Get().ViewSnapshotHistory,
+				TAttribute<FText>(),
+				TAttribute<FText>(),
+				FSlateIcon(FRevisionControlStyleManager::Get().GetStyleSetName(), "RevisionControl.SnapshotHistoryTab")
+			);
+		}
+	}));
+
+	SourceControlSection.AddDynamicEntry("CheckOutModifiedFilesTab", FNewToolMenuSectionDelegate::CreateLambda([](FToolMenuSection& InSection)
+	{
+		ISourceControlModule& SourceControlModule = ISourceControlModule::Get();
+
+		if (SourceControlModule.GetProvider().UsesCheckout())
+		{
+			InSection.AddMenuEntry(
+				FMainFrameCommands::Get().CheckOutModifiedFiles,
+				TAttribute<FText>(),
+				TAttribute<FText>(),
+				FSlateIcon(FRevisionControlStyleManager::Get().GetStyleSetName(), "RevisionControl.Actions.CheckOut")
+			);
+		}
+	}));
 
 	SourceControlSection.AddDynamicEntry("ConnectToSourceControl", FNewToolMenuSectionDelegate::CreateLambda([](FToolMenuSection& InSection)
 	{
 		ISourceControlModule& SourceControlModule = ISourceControlModule::Get();
 
-		if (ISourceControlModule::Get().IsEnabled() && ISourceControlModule::Get().GetProvider().IsAvailable())
+		if (SourceControlModule.IsEnabled() && SourceControlModule.GetProvider().IsAvailable())
 		{
 			InSection.AddMenuEntry(
 				FMainFrameCommands::Get().ChangeSourceControlSettings,
 				TAttribute<FText>(),
 				TAttribute<FText>(),
-				FSlateIcon(FAppStyle::GetAppStyleSetName(), "SourceControl.Actions.ChangeSettings")
+				FSlateIcon(FRevisionControlStyleManager::Get().GetStyleSetName(), "RevisionControl.Actions.ChangeSettings")
 			);
 		}
 		else
@@ -685,7 +726,7 @@ void FMainMenu::RegisterToolsMenu()
 				FMainFrameCommands::Get().ConnectToSourceControl,
 				TAttribute<FText>(),
 				TAttribute<FText>(),
-				FSlateIcon(FAppStyle::GetAppStyleSetName(), "SourceControl.Actions.Connect")
+				FSlateIcon(FRevisionControlStyleManager::Get().GetStyleSetName(), "RevisionControl.Actions.Connect")
 			);
 		}
 	}));
diff --git a/Engine/Source/Editor/StatusBar/Private/SourceControlMenuHelpers.cpp b/Engine/Source/Editor/StatusBar/Private/SourceControlMenuHelpers.cpp
--- a/Engine/Source/Editor/StatusBar/Private/SourceControlMenuHelpers.cpp
+++ b/Engine/Source/Editor/StatusBar/Private/SourceControlMenuHelpers.cpp
@@ -83,7 +83,9 @@ void FSourceControlCommands::RegisterCommands()
 	ActionList->MapAction(
 		CheckOutModifiedFiles,
 		FExecuteAction::CreateStatic(&FSourceControlCommands::CheckOutModifiedFiles_Clicked),
-		FCanExecuteAction::CreateStatic(&FSourceControlCommands::CheckOutModifiedFiles_CanExecute)
+		FCanExecuteAction::CreateStatic(&FSourceControlCommands::CheckOutModifiedFiles_CanExecute),
+		FIsActionChecked::CreateLambda([]() { return false; }),
+		FIsActionButtonVisible::CreateStatic(&FSourceControlCommands::CheckOutModifiedFiles_IsVisible)
 	);
 
 	ActionList->MapAction(
@@ -146,6 +148,11 @@ bool FSourceControlCommands::CheckOutModifiedFiles_CanExecute()
 	return false;
 }
 
+bool FSourceControlCommands::CheckOutModifiedFiles_IsVisible()
+{
+	return ISourceControlModule::Get().GetProvider().UsesCheckout();
+}
+
 void FSourceControlCommands::CheckOutModifiedFiles_Clicked()
 {
 	TArray<UPackage*> PackagesToSave;
@@ -222,6 +229,8 @@ void FSourceControlMenuHelpers::OnSourceControlOperationComplete(const FSourceCo
 
 TSharedRef<SWidget> FSourceControlMenuHelpers::GenerateSourceControlMenuContent()
 {
+	// NOTE: This menu should be in sync with MainMenu's SourceControl menu.
+
 	UToolMenu* SourceControlMenu = UToolMenus::Get()->RegisterMenu("StatusBar.ToolBar.SourceControl", NAME_None, EMultiBoxType::Menu, false);
 
 	FToolMenuSection& Section = SourceControlMenu->AddSection("SourceControlActions", LOCTEXT("SourceControlMenuHeadingActions", "Actions"));
@@ -265,7 +274,7 @@ TSharedRef<SWidget> FSourceControlMenuHelpers::GenerateSourceControlMenuContent(
 				FSourceControlCommands::Get().ConnectToSourceControl,
 				TAttribute<FText>(),
 				TAttribute<FText>(),
-				FSlateIcon(FRevisionControlStyleManager::GetStyleSetName(), "RevisionControl.Icon")
+				FSlateIcon(FRevisionControlStyleManager::GetStyleSetName(), "RevisionControl.Actions.Connect")
 			);
 		}
 	}));
@@ -283,14 +292,14 @@ TSharedRef<SWidget> FSourceControlMenuHelpers::GenerateCheckInComboButtonContent
 		FSourceControlCommands::Get().RevertAll,
 		TAttribute<FText>(),
 		TAttribute<FText>(),
-		FSlateIcon(FAppStyle::GetAppStyleSetName(), "SourceControl.Actions.Revert")
+		FSlateIcon(FRevisionControlStyleManager::Get().GetStyleSetName(), "RevisionControl.Actions.Revert")
 	);
 
 	Section.AddMenuEntry(
 		FSourceControlCommands::Get().ViewChangelists,
 		TAttribute<FText>(),
 		TAttribute<FText>(),
-		FSlateIcon(FAppStyle::GetAppStyleSetName(), "SourceControl.ChangelistsTab")
+		FSlateIcon(FRevisionControlStyleManager::Get().GetStyleSetName(), "RevisionControl.ChangelistsTab")
 	);
 
 	return UToolMenus::Get()->GenerateWidget("StatusBar.ToolBar.SourceControl.CheckInCombo", FToolMenuContext(FSourceControlCommands::ActionList));
diff --git a/Engine/Source/Editor/StatusBar/Private/SourceControlMenuHelpers.h b/Engine/Source/Editor/StatusBar/Private/SourceControlMenuHelpers.h
--- a/Engine/Source/Editor/StatusBar/Private/SourceControlMenuHelpers.h
+++ b/Engine/Source/Editor/StatusBar/Private/SourceControlMenuHelpers.h
@@ -26,11 +26,12 @@ private:
 	static void ConnectToSourceControl_Clicked();
 	static bool ViewChangelists_CanExecute();
 	static bool ViewChangelists_IsVisible();
+	static void ViewChangelists_Clicked();
 	static bool ViewSnapshotHistory_CanExecute();
 	static bool ViewSnapshotHistory_IsVisible();
-	static void ViewChangelists_Clicked();
 	static void ViewSnapshotHistory_Clicked();
 	static bool CheckOutModifiedFiles_CanExecute();
+	static bool CheckOutModifiedFiles_IsVisible();
 	static void CheckOutModifiedFiles_Clicked();
 	static bool RevertAllModifiedFiles_CanExecute();
 	static void RevertAllModifiedFiles_Clicked();
