[LevelEditor]
* Fixed linker errors related to LevelViewportCommands

diff --git a/Engine/Source/Editor/LevelEditor/Private/LevelViewportActions.cpp b/Engine/Source/Editor/LevelEditor/Private/LevelViewportActions.cpp
--- a/Engine/Source/Editor/LevelEditor/Private/LevelViewportActions.cpp
+++ b/Engine/Source/Editor/LevelEditor/Private/LevelViewportActions.cpp
@@ -21,6 +21,8 @@
 
 #define LOCTEXT_NAMESPACE "LevelViewportActions"
 
+UE_DEFINE_COMMANDS_TLS(FLevelViewportCommands)
+
 FLevelViewportCommands::FOnNewStatCommandAdded FLevelViewportCommands::NewStatCommandDelegate;
 
 PRAGMA_DISABLE_DEPRECATION_WARNINGS
diff --git a/Engine/Source/Editor/LevelEditor/Public/LevelViewportActions.h b/Engine/Source/Editor/LevelEditor/Public/LevelViewportActions.h
--- a/Engine/Source/Editor/LevelEditor/Public/LevelViewportActions.h
+++ b/Engine/Source/Editor/LevelEditor/Public/LevelViewportActions.h
@@ -35,10 +35,12 @@ namespace LevelViewportConfigurationNames
 	static FName OnePane("OnePane");
 }
 
+UE_DECLARE_COMMANDS_TLS(class FLevelViewportCommands, LEVELEDITOR_API)
+
 /**
  * Class containing commands for level viewport actions
  */
-class LEVELEDITOR_API FLevelViewportCommands : public TCommands<FLevelViewportCommands>
+class FLevelViewportCommands : public TCommands<FLevelViewportCommands>
 {
 public:
 
diff --git a/Engine/Source/Editor/MainFrame/MainFrame.Build.cs b/Engine/Source/Editor/MainFrame/MainFrame.Build.cs
--- a/Engine/Source/Editor/MainFrame/MainFrame.Build.cs
+++ b/Engine/Source/Editor/MainFrame/MainFrame.Build.cs
@@ -22,6 +22,7 @@ public class MainFrame : ModuleRules
 				"Engine",
 				"EngineSettings",
 				"InputCore",
+				"LevelEditor",
 				"RHI",
 				"RenderCore",
 				"Slate",
@@ -53,7 +54,6 @@ public class MainFrame : ModuleRules
 				"LauncherPlatform",
 				"GameProjectGeneration",
 				"ProjectTargetPlatformEditor",
-				"LevelEditor",
 				"Settings",
 				"SourceCodeAccess",
 				"HotReload",
@@ -65,7 +65,6 @@ public class MainFrame : ModuleRules
 				"AssetTools",
 				"GameProjectGeneration",
 				"ProjectTargetPlatformEditor",
-				"LevelEditor",
 				"SourceCodeAccess",
 				"HotReload",
 				"LocalizationDashboard", // not required but causes circular depends issues on Linux/Mac
