[Localization] Localize Copy Of when saving layout strings

diff --git a/Engine/Source/Editor/MainFrame/Private/Menus/LayoutsMenu.cpp b/Engine/Source/Editor/MainFrame/Private/Menus/LayoutsMenu.cpp
--- a/Engine/Source/Editor/MainFrame/Private/Menus/LayoutsMenu.cpp
+++ b/Engine/Source/Editor/MainFrame/Private/Menus/LayoutsMenu.cpp
@@ -630,11 +630,11 @@ void FPrivateLayoutsMenu::SaveExportLayoutCommon(const FString& InDefaultDirecto
 			// We want to avoid the duplication of the file/name/description fields, so we add "Copy of " at the beginning of the name and description fields
 			if (LayoutNames[0].ToString().Len() > 0)
 			{
-				LayoutNames[0] = FText::FromString(TEXT("Copy of ") + LayoutNames[0].ToString());
+				LayoutNames[0] = FText::Format(LOCTEXT("CopyOfLayoutName", "Copy of {0}"), FText::FromString(LayoutNames[0].ToString()));
 			}
 			if (LayoutDescriptions[0].ToString().Len() > 0)
 			{
-				LayoutDescriptions[0] = FText::FromString(TEXT("Copy of ") + LayoutDescriptions[0].ToString());
+				LayoutDescriptions[0] = FText::Format(LOCTEXT("CopyOfLayoutName", "Copy of {0}"), FText::FromString(LayoutDescriptions[0].ToString()));
 			}
 			// Create SWidget for saving the layout in its own SWindow and block the thread until it is finished
 			const TSharedRef<FSaveLayoutDialogParams> SaveLayoutDialogParams = MakeShared<FSaveLayoutDialogParams>(InDefaultDirectory, TEXT(".ini"), LayoutNames, LayoutDescriptions);
