Skein: Fixed F5 not syncing to latest.

The 'Sync Content' and 'Submit Content' buttons were previously removed from the UI, per requested. Since then, the 'Submit Content' button has been restored for Git/Subversion only. The 'Sync Content' button is still gone, which meant the F5 shortcut didn't work anymore. This changelist restores the F5 action, but does NOT restore the 'Sync Content' button in the UI.

#rnx

diff --git a/Engine/Source/Editor/MainFrame/Private/Frame/MainFrameActions.cpp b/Engine/Source/Editor/MainFrame/Private/Frame/MainFrameActions.cpp
--- a/Engine/Source/Editor/MainFrame/Private/Frame/MainFrameActions.cpp
+++ b/Engine/Source/Editor/MainFrame/Private/Frame/MainFrameActions.cpp
@@ -114,6 +114,9 @@ void FMainFrameCommands::RegisterCommands()
 	UI_COMMAND( SubmitContent, "Submit Content", "Opens a dialog with check in options for content and levels.", EUserInterfaceActionType::Button, FInputChord() );
 	ActionList->MapAction( SubmitContent, FExecuteAction::CreateLambda([]() { FSourceControlWindows::ChoosePackagesToCheckIn(); }), FCanExecuteAction::CreateStatic( &FMainFrameActionCallbacks::CanSubmitContent ) );
 
+	UI_COMMAND( SyncContent, "Sync Content", "Saves all unsaved levels and assets to disk and then downloads the latest versions from revision control.", EUserInterfaceActionType::Button, FInputChord() );
+	ActionList->MapAction( SyncContent, FExecuteAction::CreateLambda([]() {  FSourceControlWindows::SyncLatest(); }), FCanExecuteAction::CreateStatic(&FMainFrameActionCallbacks::CanSyncContent) );
+
 	UI_COMMAND( CheckOutModifiedFiles, "Check Out Modified Files", "Opens a dialog to check out any assets which have been modified.", EUserInterfaceActionType::Button, FInputChord());
 	ActionList->MapAction( CheckOutModifiedFiles, FExecuteAction::CreateStatic( &FMainFrameActionCallbacks::CheckOutModifiedFiles ), FCanExecuteAction::CreateStatic( &FMainFrameActionCallbacks::CanCheckOutModifiedFiles ) );
 
@@ -380,6 +383,11 @@ bool FMainFrameActionCallbacks::CanSubmitContent()
 	return GLevelEditorModeTools().CanAutoSave() && FSourceControlWindows::CanChoosePackagesToCheckIn();
 }
 
+bool FMainFrameActionCallbacks::CanSyncContent()
+{
+	return GLevelEditorModeTools().CanAutoSave() && FSourceControlWindows::CanSyncLatest();
+}
+
 TArray<FRecentProjectFile> FMainFrameActionCallbacks::RecentProjects;
 
 void FMainFrameActionCallbacks::CacheProjectNames()
diff --git a/Engine/Source/Editor/MainFrame/Private/Frame/MainFrameActions.h b/Engine/Source/Editor/MainFrame/Private/Frame/MainFrameActions.h
--- a/Engine/Source/Editor/MainFrame/Private/Frame/MainFrameActions.h
+++ b/Engine/Source/Editor/MainFrame/Private/Frame/MainFrameActions.h
@@ -29,6 +29,7 @@ public:
 	TSharedPtr< FUICommandInfo > ViewChangelists;
 	TSharedPtr< FUICommandInfo > ViewSnapshotHistory;
 	TSharedPtr< FUICommandInfo > SubmitContent;
+	TSharedPtr< FUICommandInfo > SyncContent;
 	TSharedPtr< FUICommandInfo > CheckOutModifiedFiles;
 	TSharedPtr< FUICommandInfo > ConnectToSourceControl;
 	TSharedPtr< FUICommandInfo > ChangeSourceControlSettings;
@@ -114,6 +115,9 @@ public:
 	/** Determine whether we are allowed to submit content at this moment */
 	static bool CanSubmitContent();
 
+	/** Determine whether we are allocated to sync content at this moment */
+	static bool CanSyncContent();
+
 	/** Opens a dialog to choose packages to save */
 	static void ChoosePackagesToSave();
 
