SourceControl: Fixed 'Submit Content' being available when source control is disabled.

#rnx

diff --git a/Engine/Source/Editor/MainFrame/Private/Menus/MainMenu.cpp b/Engine/Source/Editor/MainFrame/Private/Menus/MainMenu.cpp
--- a/Engine/Source/Editor/MainFrame/Private/Menus/MainMenu.cpp
+++ b/Engine/Source/Editor/MainFrame/Private/Menus/MainMenu.cpp
@@ -659,7 +659,8 @@ void FMainMenu::RegisterToolsMenu()
 		ISourceControlModule& SourceControlModule = ISourceControlModule::Get();
 
 		// Access to the legacy 'Submit Files' dialog. Git / Subversion require this.
-		if (!SourceControlModule.GetProvider().UsesChangelists() && SourceControlModule.GetProvider().GetName() != TEXT("Unreal Revision Control"))
+		if (SourceControlModule.IsEnabled() && !SourceControlModule.GetProvider().UsesChangelists() &&
+			SourceControlModule.GetProvider().GetName() != TEXT("Unreal Revision Control"))
 		{
 			InSection.AddMenuEntry(
 				FMainFrameCommands::Get().SubmitContent,
diff --git a/Engine/Source/Editor/StatusBar/Private/SourceControlMenuHelpers.cpp b/Engine/Source/Editor/StatusBar/Private/SourceControlMenuHelpers.cpp
--- a/Engine/Source/Editor/StatusBar/Private/SourceControlMenuHelpers.cpp
+++ b/Engine/Source/Editor/StatusBar/Private/SourceControlMenuHelpers.cpp
@@ -141,7 +141,7 @@ bool FSourceControlCommands::SubmitContent_IsVisible()
 		return false;
 	}
 
-	return true;
+	return ISourceControlModule::Get().IsEnabled();
 }
 
 bool FSourceControlCommands::ViewSnapshotHistory_CanExecute()
