Fix undo history tab not being restored after restarting the editor.

The problem was the undo history module did not get loaded at all. Only once you click the Edit > Undo History entry was the module loaded. This change loads it when that menu is registered, which seems to be early enough.

diff --git a/Engine/Source/Editor/MainFrame/Private/Menus/MainMenu.cpp b/Engine/Source/Editor/MainFrame/Private/Menus/MainMenu.cpp
--- a/Engine/Source/Editor/MainFrame/Private/Menus/MainMenu.cpp
+++ b/Engine/Source/Editor/MainFrame/Private/Menus/MainMenu.cpp
@@ -216,6 +216,8 @@ void FMainMenu::RegisterEditMenu()
 			FSlateIcon(FAppStyle::GetAppStyleSetName(), "UndoHistory.TabIcon"),
 			FUIAction(FExecuteAction::CreateLambda([](){  IUndoHistoryEditorModule::Get().ExecuteOpenUndoHistory(); } ))
 			);
+		// IUndoHistoryEditorModule registers the nomad tab spawner. Ensure it's registered before SLevelEditor::RestoreContentArea restores the layout.
+		IUndoHistoryEditorModule::Get();
 	}
 
 	{
diff --git a/Engine/Source/Editor/UndoHistoryEditor/Private/UndoHistoryEditorModule.cpp b/Engine/Source/Editor/UndoHistoryEditor/Private/UndoHistoryEditorModule.cpp
--- a/Engine/Source/Editor/UndoHistoryEditor/Private/UndoHistoryEditorModule.cpp
+++ b/Engine/Source/Editor/UndoHistoryEditor/Private/UndoHistoryEditorModule.cpp
@@ -33,11 +33,6 @@ public:
 		FGlobalTabmanager::Get()->UnregisterNomadTabSpawner(UndoHistoryTabName);
 	}
 	
-	virtual bool SupportsDynamicReloading() override
-	{
-		return true;
-	}
-
 	virtual void ExecuteOpenUndoHistory() override
 	{
 		FGlobalTabmanager::Get()->TryInvokeTab(UndoHistoryTabName);
