[QoL Improvements] Copy and Save Layout naming issues

- Storing Layout Names and Descriptions as FName instead of FText (prevents LOCTEXT macro from coming up in New Layout prompt)

- Reads Layout Name and Description from `.json` file instead of `.ini`: this allows Layout Save/Load menus to show the actual layout name, instead of the default one.

- When overriding an existing Layout ("Save as") the layout name is kept, instead of clearing it or overwriting it.

- Both `.ini` and `.json` files will be removed when removing a layout from the menu

- New Layout dialog name and tooltip defaults are now changed to "My New Layout" and "Tooltip for My New Layout"

--- Some notes about Layout config files ---

Each layout config info is stored in 2 files:
- `.ini` file, contains default layout name
- `.json` file, contains actual Display Name (e.g.  `[REDACTED_USER_HOME]\AppData\Local\UnrealEngine\Editor`)

Current CL is reading Display Name from json file so that layouts pre-dating it can still show the right Display Name in the Save/Load menus.

--- Some notes about the override fix  ---

In order to have the "Override" work, without huge changes to the code, the system loads the Layout config which will be overriden, before actually overriding it. This allows the proper Display Name to be kept.
Other solutions might possibly involve big changes to the code, and maybe this is not desirable it right now - are going to keep both `.ini` and `.json` files in the future?

diff --git a/Engine/Source/Editor/MainFrame/Private/Menus/LayoutsMenu.cpp b/Engine/Source/Editor/MainFrame/Private/Menus/LayoutsMenu.cpp
--- a/Engine/Source/Editor/MainFrame/Private/Menus/LayoutsMenu.cpp
+++ b/Engine/Source/Editor/MainFrame/Private/Menus/LayoutsMenu.cpp
@@ -104,8 +104,6 @@ public:
 	 */
 	static bool CheckAskUserToClosePIESIE(const FText& InitialMessage);
 
-	static FText GenerateLocalizedTextForFile(const FText& InText);
-
 	static void SaveExportLayoutCommon(const FString& InDefaultDirectory, const bool bMustBeSavedInDefaultDirectory, const FText& InWhatIsThis,
 		const bool bShouldAskBeforeCleaningLayoutNameAndDescriptionFields);
 
@@ -289,8 +287,8 @@ bool FPrivateLayoutsMenu::TrySaveLayoutOrWarnInternal(const FString& InSourceFil
 		if (bCleanLayoutNameAndDescriptionFields)
 		{
 			// Update fields
-			FLayoutSaveRestore::SaveSectionToConfig(GEditorLayoutIni, "LayoutName", FText::FromString(""));
-			FLayoutSaveRestore::SaveSectionToConfig(GEditorLayoutIni, "LayoutDescription", FText::FromString(""));
+			FLayoutSaveRestore::SaveSectionToConfig(GEditorLayoutIni, "LayoutName", TEXT(""));
+			FLayoutSaveRestore::SaveSectionToConfig(GEditorLayoutIni, "LayoutDescription", TEXT(""));
 
 			// Flush file
 			const bool bRead = true;
@@ -350,7 +348,7 @@ FText FPrivateLayoutsMenu::GetTooltipTextInternal(const FText& InDisplayName, co
 	}
 	else
 	{
-		Tooltip = FText::Format(LOCTEXT("LayoutNameFmt", "Description:\n{0}.\n\nFull file path:\n{1}"), InLayoutName, FText::FromString(InLayoutFilePath));
+		Tooltip = FText::Format(LOCTEXT("LayoutNameFmt", "Description:\n{0}\n\nFull file path:\n{1}"), InLayoutName, FText::FromString(InLayoutFilePath));
 	}
 	return Tooltip;
 }
@@ -584,32 +582,6 @@ bool FPrivateLayoutsMenu::CheckAskUserToClosePIESIE(const FText& InitialMessage)
 	return true;
 }
 
-FText FPrivateLayoutsMenu::GenerateLocalizedTextForFile(const FText& InText)
-{
-	// Proper FText to FString
-	FString StringSimulatingText;
-	FTextStringHelper::WriteToBuffer(StringSimulatingText, InText);
-	// Sanitize text (truncate if too big)
-	FString SanitazedTruncatedText = (StringSimulatingText.Len() < 100 ? StringSimulatingText : StringSimulatingText.Left(100));
-	FSaveLayoutDialogUtils::SanitizeText(SanitazedTruncatedText);
-	// Create full file path
-	if (StringSimulatingText.Len() < 10 || StringSimulatingText.Left(9) != TEXT("NSLOCTEXT"))
-	{
-		FString StringSimulatingTextRecreated =
-			// Namespace
-			TEXT("NSLOCTEXT(\"LayoutNamespace\", "
-			// Key
-			"\"") + SanitazedTruncatedText + TEXT("\", "
-			// Source string
-			"\"") + StringSimulatingText + TEXT("\")");
-		return FText::FromString(StringSimulatingTextRecreated);
-	}
-	else
-	{
-		return FText::FromString(StringSimulatingText);
-	}
-}
-
 void FPrivateLayoutsMenu::SaveExportLayoutCommon(const FString& InDefaultDirectory, const bool bMustBeSavedInDefaultDirectory, const FText& InWhatIsThis,
 	const bool bShouldAskBeforeCleaningLayoutNameAndDescriptionFields)
 {
@@ -623,19 +595,18 @@ void FPrivateLayoutsMenu::SaveExportLayoutCommon(const FString& InDefaultDirecto
 		// "Save Layout As..."
 		if (bMustBeSavedInDefaultDirectory)
 		{
-			TArray<FText> LayoutNames;
-			LayoutNames.Emplace(FLayoutSaveRestore::LoadSectionFromConfig(GEditorLayoutIni, "LayoutName"));
-			TArray<FText> LayoutDescriptions;
-			LayoutDescriptions.Emplace(FLayoutSaveRestore::LoadSectionFromConfig(GEditorLayoutIni, "LayoutDescription"));
+			TArray<FName> LayoutNames;
+			LayoutNames.Emplace(FLayoutSaveRestore::LoadSectionFromConfig(GEditorLayoutIni, "LayoutName").ToString());
+			TArray<FName> LayoutDescriptions;
+			LayoutDescriptions.Emplace(FLayoutSaveRestore::LoadSectionFromConfig(GEditorLayoutIni, "LayoutDescription").ToString());
 			// We want to avoid the duplication of the file/name/description fields, so we add "Copy of " at the beginning of the name and description fields
-			if (LayoutNames[0].ToString().Len() > 0)
-			{
-				LayoutNames[0] = FText::Format(LOCTEXT("CopyOfLayoutName", "Copy of {0}"), FText::FromString(LayoutNames[0].ToString()));
-			}
-			if (LayoutDescriptions[0].ToString().Len() > 0)
-			{
-				LayoutDescriptions[0] = FText::Format(LOCTEXT("CopyOfLayoutName", "Copy of {0}"), FText::FromString(LayoutDescriptions[0].ToString()));
-			}
+
+			const FText NewLayoutName = LOCTEXT("NewLayoutName", "My New Layout");
+			LayoutNames[0] = FName(NewLayoutName.ToString());
+
+			const FText NewLayoutDescription = LOCTEXT("NewLayoutDescription", "Tooltip for My New Layout");
+			LayoutDescriptions[0] = FName(NewLayoutDescription.ToString());
+
 			// Create SWidget for saving the layout in its own SWindow and block the thread until it is finished
 			const TSharedRef<FSaveLayoutDialogParams> SaveLayoutDialogParams = MakeShared<FSaveLayoutDialogParams>(InDefaultDirectory, TEXT(".ini"), LayoutNames, LayoutDescriptions);
 			bWasDialogOpened = FSaveLayoutDialogUtils::CreateSaveLayoutAsDialogInStandaloneWindow(SaveLayoutDialogParams);
@@ -650,11 +621,11 @@ void FPrivateLayoutsMenu::SaveExportLayoutCommon(const FString& InDefaultDirecto
 				checkf(LayoutNames.Num() == LayoutDescriptions.Num(), TEXT("There should be the same number of LayoutNames and LayoutDescriptions."));
 				for (int32 Index = 0; Index < LayoutNames.Num(); ++Index)
 				{
-					const FText& LayoutNameAsTextText = GenerateLocalizedTextForFile(LayoutNames[Index]);
-					const FText& LayoutDescriptionAsTextText = GenerateLocalizedTextForFile(LayoutDescriptions[Index]);
+					const FString LayoutName = LayoutNames[Index].ToString();
+					const FString LayoutDescription = LayoutDescriptions[Index].ToString();
 					// Update fields
-					FLayoutSaveRestore::SaveSectionToConfig(GEditorLayoutIni, "LayoutName", LayoutNameAsTextText);
-					FLayoutSaveRestore::SaveSectionToConfig(GEditorLayoutIni, "LayoutDescription", LayoutDescriptionAsTextText);
+					FLayoutSaveRestore::SaveSectionToConfig(GEditorLayoutIni, "LayoutName", LayoutName);
+					FLayoutSaveRestore::SaveSectionToConfig(GEditorLayoutIni, "LayoutDescription", LayoutDescription);
 					// Flush file
 					const bool bRead = true;
 					GConfig->Flush(bRead, GEditorLayoutIni);
@@ -826,7 +797,7 @@ void FLayoutsMenuLoad::ReloadCurrentLayout()
 	}
 }
 
-void FLayoutsMenuLoad::LoadLayout(const FString& InLayoutPath)
+void FLayoutsMenuLoad::LoadLayout(const FString& InLayoutPath, bool bInSkipLayoutReload)
 {
 	// Replace main layout with desired one
 	const FString& SourceFilePath = InLayoutPath;
@@ -836,7 +807,7 @@ void FLayoutsMenuLoad::LoadLayout(const FString& InLayoutPath)
 	const bool SucessfullySaved = FPrivateLayoutsMenu::TrySaveLayoutOrWarnInternal(SourceFilePath, TargetFilePath, LOCTEXT("LoadLayoutText", "loading the layout"),
 		bCleanLayoutNameAndDescriptionFieldsIfNoSameValues, bShouldAskBeforeCleaningLayoutNameAndDescriptionFields);
 	// Reload current layout
-	if (SucessfullySaved)
+	if (SucessfullySaved && !bInSkipLayoutReload)
 	{
 		FLayoutsMenuLoad::ReloadCurrentLayout();
 	}
@@ -983,6 +954,13 @@ void FLayoutsMenuSave::OverrideLayout(const int32 InLayoutIndex, const FLayoutsM
 	{
 		return;
 	}
+
+	// Loading the layout which is about to be overriden, so we can store its proper name.
+	// The code which overrides the current layout ends up using the Display Name and Description of the currently loaded layout.
+	// That's not always the desired one. Making sure of it by loading it.
+	constexpr bool bSkipReload = true;
+	FLayoutsMenuLoad::LoadLayout(DesiredLayoutFullPath, bSkipReload);
+
 	// Target and source files
 	const FString& SourceFilePath = GEditorLayoutIni;
 	const FString& TargetFilePath = DesiredLayoutFullPath;
@@ -1066,6 +1044,17 @@ bool FLayoutsMenuRemove::CanRemoveChooseLayout(const FLayoutsMenu::ELayoutsType
 	return FPrivateLayoutsMenu::CanChooseLayoutWhenWriteInternal(InLayoutsType);
 }
 
+static FString GetLayoutJsonFileName(const FString& InConfigFileName)
+{
+	const FString JsonFileName = FPaths::GetBaseFilename(InConfigFileName) + TEXT(".json");
+#ifdef UE_SAVED_DIR_OVERRIDE
+	const FString UserSettingsPath = FPaths::Combine(FPlatformProcess::UserSettingsDir(), TEXT(PREPROCESSOR_TO_STRING(UE_SAVED_DIR_OVERRIDE)), TEXT("Editor"), JsonFileName);
+#else
+	const FString UserSettingsPath = FPaths::Combine(FPlatformProcess::UserSettingsDir(), FApp::GetEpicProductIdentifier(), TEXT("Editor"), JsonFileName);
+#endif
+	return UserSettingsPath;
+}
+
 void FLayoutsMenuRemove::RemoveLayout(const int32 InLayoutIndex, const FLayoutsMenu::ELayoutsType InLayoutsType)
 {
 	// Default layouts should never be modified, so this function should never be called
@@ -1086,6 +1075,13 @@ void FLayoutsMenuRemove::RemoveLayout(const int32 InLayoutIndex, const FLayoutsM
 	}
 	// Remove layout
 	FPlatformFileManager::Get().GetPlatformFile().DeleteFile(*DesiredLayoutFullPath);
+
+	// Remove JSON layout as well
+	const FString TargetJsonFilepath = GetLayoutJsonFileName(LayoutIniFileNames[InLayoutIndex]);
+	if (FLayoutSaveRestore::IsValidConfig(TargetJsonFilepath))
+	{
+		FPlatformFileManager::Get().GetPlatformFile().DeleteFile(*TargetJsonFilepath);
+	}
 }
 
 void FLayoutsMenuRemove::RemoveUserLayouts()
@@ -1116,6 +1112,13 @@ void FLayoutsMenuRemove::RemoveUserLayouts()
 			{
 				FPlatformFileManager::Get().GetPlatformFile().DeleteFile(*LayoutFilePath);
 			}
+
+			// Remove JSON layout as well
+			const FString TargetJsonFilepath = GetLayoutJsonFileName(UserLayoutIniFileName);
+			if (FLayoutSaveRestore::IsValidConfig(TargetJsonFilepath))
+			{
+				FPlatformFileManager::Get().GetPlatformFile().DeleteFile(*TargetJsonFilepath);
+			}
 		}
 	}
 	// If no files to remove, warn user
diff --git a/Engine/Source/Editor/MainFrame/Private/Menus/LayoutsMenu.h b/Engine/Source/Editor/MainFrame/Private/Menus/LayoutsMenu.h
--- a/Engine/Source/Editor/MainFrame/Private/Menus/LayoutsMenu.h
+++ b/Engine/Source/Editor/MainFrame/Private/Menus/LayoutsMenu.h
@@ -79,8 +79,9 @@ public:
 	 * Static
 	 * Load the visual layout state of the editor from an existing layout profile ini file, given its file path.
 	 * @param InLayoutPath File path associated with the desired layout profile ini file to be read/written.
+	 * @param bInSkipLayoutReload Prevents FLayoutsMenuLoad::ReloadCurrentLayout() from being called after loading the specified layout profile.
 	 */
-	static void LoadLayout(const FString& InLayoutPath);
+	static void LoadLayout(const FString& InLayoutPath, bool bInSkipLayoutReload = false);
 
 	/**
 	 * Static
diff --git a/Engine/Source/Editor/MainFrame/Private/Menus/SaveLayoutDialog.cpp b/Engine/Source/Editor/MainFrame/Private/Menus/SaveLayoutDialog.cpp
--- a/Engine/Source/Editor/MainFrame/Private/Menus/SaveLayoutDialog.cpp
+++ b/Engine/Source/Editor/MainFrame/Private/Menus/SaveLayoutDialog.cpp
@@ -102,9 +102,9 @@ private:
 
 	FString FinalFilePath;
 
-	FText CurrentlyEnteredLayoutName;
+	FName CurrentlyEnteredLayoutName;
 
-	FText CurrentlyEnteredLayoutDescription;
+	FName CurrentlyEnteredLayoutDescription;
 
 	/** The error text from the last validity check */
 	FText LastInputValidityErrorText;
@@ -115,7 +115,7 @@ private:
 
 
 
-FSaveLayoutDialogParams::FSaveLayoutDialogParams(const FString& InDefaultDirectory, const FString& InFileExtension, const TArray<FText>& InLayoutNames, const TArray<FText>& InLayoutDescriptions)
+FSaveLayoutDialogParams::FSaveLayoutDialogParams(const FString& InDefaultDirectory, const FString& InFileExtension, const TArray<FName>& InLayoutNames, const TArray<FName>& InLayoutDescriptions)
 	: DefaultDirectory(InDefaultDirectory)
 	, FileExtension(InFileExtension)
 	, LayoutNames(InLayoutNames)
@@ -239,7 +239,7 @@ void SSaveLayoutDialog::Construct(const FArguments& InArgs, const TSharedRef<FSa
 					.Padding(0, 0, 0, 8)
 					[
 						SNew(SEditableTextBox)
-						.Text(CurrentlyEnteredLayoutName)
+						.Text(FText::FromString(CurrentlyEnteredLayoutName.ToString()))
 						.OnTextCommitted(this, &SSaveLayoutDialog::OnLayoutNameTextCommited)
 						.OnTextChanged(this, &SSaveLayoutDialog::OnLayoutNameTextCommited, ETextCommit::Default)
 						.SelectAllTextWhenFocused(true)
@@ -248,7 +248,7 @@ void SSaveLayoutDialog::Construct(const FArguments& InArgs, const TSharedRef<FSa
 					.Padding(0, 0, 0, 8)
 					[
 						SNew(SEditableTextBox)
-						.Text(CurrentlyEnteredLayoutDescription)
+						.Text(FText::FromString(CurrentlyEnteredLayoutDescription.ToString()))
 						.HintText(LOCTEXT("LayoutDescriptionInputBoxHintText", "Optional"))
 						.OnTextCommitted(this, &SSaveLayoutDialog::OnLayoutDescriptionTextCommited)
 						.OnTextChanged(this, &SSaveLayoutDialog::OnLayoutDescriptionTextCommited, ETextCommit::Default)
@@ -324,7 +324,7 @@ FText SSaveLayoutDialog::GetNameErrorLabelText() const
 
 void SSaveLayoutDialog::OnLayoutNameTextCommited(const FText& InText, ETextCommit::Type InCommitType)
 {
-	CurrentlyEnteredLayoutName = InText;
+	CurrentlyEnteredLayoutName = FName(InText.ToString());
 	UpdateInputValidity();
 
 	if (InCommitType == ETextCommit::OnEnter)
@@ -335,7 +335,7 @@ void SSaveLayoutDialog::OnLayoutNameTextCommited(const FText& InText, ETextCommi
 
 void SSaveLayoutDialog::OnLayoutDescriptionTextCommited(const FText& InText, ETextCommit::Type InCommitType)
 {
-	CurrentlyEnteredLayoutDescription = InText;
+	CurrentlyEnteredLayoutDescription = FName(InText.ToString());
 
 	if (InCommitType == ETextCommit::OnEnter)
 	{
diff --git a/Engine/Source/Editor/MainFrame/Private/Menus/SaveLayoutDialog.h b/Engine/Source/Editor/MainFrame/Private/Menus/SaveLayoutDialog.h
--- a/Engine/Source/Editor/MainFrame/Private/Menus/SaveLayoutDialog.h
+++ b/Engine/Source/Editor/MainFrame/Private/Menus/SaveLayoutDialog.h
@@ -16,8 +16,8 @@ public:
 	/**
 	 * The constructor only contains the variables that FSaveLayoutDialogUtils::CreateSaveLayoutAsDialogInStandaloneWindow will need. The rest will be overriden.
 	 */
-	FSaveLayoutDialogParams(const FString& InDefaultDirectory, const FString& InFileExtension, const TArray<FText>& LayoutNames = TArray<FText>(),
-		const TArray<FText>& LayoutDescriptions = TArray<FText>());
+	FSaveLayoutDialogParams(const FString& InDefaultDirectory, const FString& InFileExtension, const TArray<FName>& LayoutNames = TArray<FName>(),
+		const TArray<FName>& LayoutDescriptions = TArray<FName>());
 
 	/**
 	 * The directory where the layout will be saved.
@@ -42,12 +42,12 @@ public:
 	/**
 	 * The displayed layout name that will be visualized in Unreal. Its initial value will be used as default.
 	 */
-	TArray<FText> LayoutNames;
+	TArray<FName> LayoutNames;
 
 	/**
 	 * The displayed layout description that will be visualized in Unreal. Its initial value will be used as default.
 	 */
-	TArray<FText> LayoutDescriptions;
+	TArray<FName> LayoutDescriptions;
 };
 
 /**
diff --git a/Engine/Source/Runtime/Slate/Private/Framework/Docking/LayoutService.cpp b/Engine/Source/Runtime/Slate/Private/Framework/Docking/LayoutService.cpp
--- a/Engine/Source/Runtime/Slate/Private/Framework/Docking/LayoutService.cpp
+++ b/Engine/Source/Runtime/Slate/Private/Framework/Docking/LayoutService.cpp
@@ -350,16 +350,18 @@ TSharedRef<FTabManager::FLayout> FLayoutSaveRestore::LoadFromConfigPrivate(const
 
 void FLayoutSaveRestore::SaveSectionToConfig(const FString& InConfigFileName, const FString& InSectionName, const FText& InSectionValue)
 {
-	FString StrValue;
-	FTextStringHelper::WriteToBuffer(StrValue, InSectionValue);
+	SaveSectionToConfig(InConfigFileName, InSectionName, InSectionValue.ToString());
+}
 
-	GConfig->SetString(GetEditorLayoutsSectionName(), *InSectionName, *StrValue, InConfigFileName);
+void FLayoutSaveRestore::SaveSectionToConfig(const FString& InConfigFileName, const FString& InSectionName, const FString& InSectionValue)
+{
+	GConfig->SetString(GetEditorLayoutsSectionName(), *InSectionName, *InSectionValue, InConfigFileName);
 
 	const FString JsonFileName = GetLayoutJsonFileName(InConfigFileName);
 	TSharedPtr<FJsonObject> JsonObject = LoadJsonFile(JsonFileName);
 	if (JsonObject.IsValid())
 	{
-		JsonObject->SetStringField(InSectionName, *StrValue);
+		JsonObject->SetStringField(InSectionName, *InSectionValue);
 		SaveJsonFile(JsonFileName, JsonObject);
 	}
 }
@@ -401,11 +403,23 @@ bool FLayoutSaveRestore::DuplicateConfig(const FString& SourceConfigFileName, co
 		return false;
 	}
 
-	// convert this layout to a JSON file
-	TArray<FString> SectionPairs;
-	GetConfigSection(SectionPairs, TargetConfigFileName);
 
-	TSharedPtr<FJsonObject> RootObject = ConvertSectionToJson(SectionPairs);
+	TSharedPtr<FJsonObject> RootObject;
+
+	// The right layout config to be duplicated is already in a json file
+	// Let's load it and use that as source for duplication
+	if (TSharedPtr<FJsonObject> SourceJson = LoadJsonFile(GetLayoutJsonFileName(SourceConfigFileName)))
+	{
+		RootObject = SourceJson;
+	}
+	else
+	{
+		// If we end up here, there is no json file for the specified ini layout:
+		// convert this layout to a JSON file
+		TArray<FString> SectionPairs;
+		GetConfigSection(SectionPairs, TargetConfigFileName);
+		RootObject = ConvertSectionToJson(SectionPairs);
+	}
 
 	const FString TargetJsonFilename = GetLayoutJsonFileName(TargetConfigFileName);
 	SaveJsonFile(TargetJsonFilename, RootObject);
diff --git a/Engine/Source/Runtime/Slate/Public/Framework/Docking/LayoutService.h b/Engine/Source/Runtime/Slate/Public/Framework/Docking/LayoutService.h
--- a/Engine/Source/Runtime/Slate/Public/Framework/Docking/LayoutService.h
+++ b/Engine/Source/Runtime/Slate/Public/Framework/Docking/LayoutService.h
@@ -38,7 +38,7 @@ struct FLayoutSaveRestore
 		const EOutputCanBeNullptr InPrimaryAreaOutputCanBeNullptr, TArray<FString>& OutRemovedOlderLayoutVersions);
 
 	/**
-	 * Write the desired FText value into the desired section of a named config file.
+	 * Write the desired value into the desired section of a named config file.
 	 * This function should only be used to save localization names (e.g., LayoutName, LayoutDescription).
 	 * For saving the FTabManager::FLayout, check SaveToConfig.
 	 *
@@ -46,6 +46,9 @@ struct FLayoutSaveRestore
 	 * @param InSectionName the section name where to save the value.
 	 * @param InSectionValue the value to save.
 	 */
+	static SLATE_API void SaveSectionToConfig(const FString& InConfigFileName, const FString& InSectionName, const FString& InSectionValue);
+
+	UE_DEPRECATED(5.7, "SaveSectionToConfig using FText value is deprecated. Please use the SaveSectionToConfig function using FString.")
 	static SLATE_API void SaveSectionToConfig(const FString& InConfigFileName, const FString& InSectionName, const FText& InSectionValue);
 
 	/**
