[HomeScreen] Add HomeScreen as a Normal tab under a CVar to test functionality

diff --git a/Engine/Source/Editor/EditorStyle/Private/StarshipStyle.cpp b/Engine/Source/Editor/EditorStyle/Private/StarshipStyle.cpp
--- a/Engine/Source/Editor/EditorStyle/Private/StarshipStyle.cpp
+++ b/Engine/Source/Editor/EditorStyle/Private/StarshipStyle.cpp
@@ -247,6 +247,7 @@ void FStarshipEditorStyle::FStyle::Initialize()
 	SetupSourceControlStyles();
 	SetupAutomationStyles();
 	SetupUMGEditorStyles();
+	SetupHomeScreenStyles();
 	SetupMyBlueprintStyles();
 	SetupStatusBarStyle();
 	SetupColorPickerStyle();
@@ -8516,6 +8517,109 @@ void FStarshipEditorStyle::FStyle::SetupUMGEditorStyles()
 	Set("UMGEditor.AnimTabIcon", new IMAGE_BRUSH_SVG("Starship/MainToolbar/cinematics", Icon16x16));
 }
 
+void FStarshipEditorStyle::FStyle::SetupHomeScreenStyles()
+{
+	Set("HomeScreen.UnrealLogo", new IMAGE_BRUSH_SVG("Starship/HomeScreen/UnrealLogo", FVector2D(99.f, 120.f)));
+	Set("HomeScreen.Home", new IMAGE_BRUSH_SVG("Starship/HomeScreen/Home", Icon20x20));
+	Set("HomeScreen.News", new IMAGE_BRUSH_SVG("Starship/HomeScreen/News", Icon20x20));
+	Set("HomeScreen.Rocket", new IMAGE_BRUSH_SVG("Starship/HomeScreen/Rocket", Icon20x20));
+	Set("HomeScreen.Archive", new IMAGE_BRUSH_SVG("Starship/HomeScreen/Archive", Icon20x20));
+	Set("HomeScreen.Forum", new IMAGE_BRUSH_SVG("Starship/HomeScreen/Forum", Icon20x20));
+	Set("HomeScreen.Documentation", new IMAGE_BRUSH_SVG("Starship/HomeScreen/Documentation", Icon20x20));
+	Set("HomeScreen.Tutorial", new IMAGE_BRUSH_SVG("Starship/HomeScreen/Tutorial", Icon20x20));
+	Set("HomeScreen.Roadmap", new IMAGE_BRUSH_SVG("Starship/HomeScreen/Roadmap", Icon20x20));
+	Set("HomeScreen.ReleaseNotes", new IMAGE_BRUSH_SVG("Starship/HomeScreen/ReleaseNotes", Icon20x20));
+	Set("HomeScreen.PlusCircle", new IMAGE_BRUSH_SVG("Starship/HomeScreen/PlusCircleEmpty", Icon20x20));
+	Set("HomeScreen.FolderOpen", new IMAGE_BRUSH_SVG("Starship/HomeScreen/FolderOpen", Icon20x20));
+
+	Set("SocialMedia.Facebook", new IMAGE_BRUSH_SVG("Starship/HomeScreen/Facebook", Icon20x20));
+	Set("SocialMedia.Twitter", new IMAGE_BRUSH_SVG("Starship/HomeScreen/Twitter", Icon20x20));
+	Set("SocialMedia.YouTube", new IMAGE_BRUSH_SVG("Starship/HomeScreen/YouTube", Icon20x20));
+	Set("SocialMedia.Twitch", new IMAGE_BRUSH_SVG("Starship/HomeScreen/Twitch", Icon20x20));
+	Set("SocialMedia.Instagram", new IMAGE_BRUSH_SVG("Starship/HomeScreen/Instagram", Icon20x20));
+
+	Set("HomeScreen.Background", new IMAGE_BRUSH("Starship/HomeScreen/Background", FVector2D(4096.f, 2304.f)));
+	Set("HomeScreen.FakeWeb", new IMAGE_BRUSH("Starship/HomeScreen/FakeWeb", FVector2D(2560.f, 1440.f)));
+	Set("HomeScreen.RecentProjects.Background", new FSlateRoundedBoxBrush(COLOR("#FFFFFF0D"), 16.f, FStyleColors::White25, 1.f));
+
+	Set("HomeScreen.NoInternet.Border", new FSlateRoundedBoxBrush(COLOR("#FFFFFF0D"), 8.f, COLOR("#FFFFFF59"), 1.f));
+	Set("HomeScreen.NoInternet.Icon", new IMAGE_BRUSH_SVG("Starship/HomeScreen/NoInternet", FVector2D(24.f, 24.f)));
+
+	FComboButtonStyle LoadOnStartupStyle = FAppStyle::Get().GetWidgetStyle<FComboButtonStyle>("ComboButton");
+	LoadOnStartupStyle.ButtonStyle = FButtonStyle()
+		.SetNormal(FSlateRoundedBoxBrush(COLOR("#FFFFFF0D"), 6.f, COLOR("#FFFFFF59"), 1.f))
+		.SetHovered(FSlateRoundedBoxBrush(COLOR("#FFFFFF0D"), 6.f, COLOR("#FFFFFF59"), 1.f))
+		.SetPressed(FSlateRoundedBoxBrush(COLOR("#FFFFFF0D"), 6.f, COLOR("#FFFFFF59"), 1.f))
+		.SetHoveredForeground(FStyleColors::White)
+		.SetNormalPadding(FMargin(12.0f, 4.0f))
+		.SetPressedPadding(FMargin(12.0f, 4.0f));
+
+
+	LoadOnStartupStyle.SetMenuBorderBrush(FSlateRoundedBoxBrush(FStyleColors::Panel, 6.f, FStyleColors::White25, 2.f));
+	LoadOnStartupStyle.SetMenuBorderPadding(0.f);
+
+	FButtonStyle LoadOnStartupMenuButtonStyle = FAppStyle::Get().GetWidgetStyle<FButtonStyle>("NoBorder");
+	LoadOnStartupMenuButtonStyle.SetHovered(FSlateRoundedBoxBrush(FStyleColors::Hover, 4.0f));
+	Set("HomeScreen.ComboButton", LoadOnStartupStyle);
+	Set("HomeScreen.ComboButton.MenuButton", LoadOnStartupMenuButtonStyle);
+	Set("HomeScreen.ComboButton.MenuBorder", new FSlateRoundedBoxBrush(FStyleColors::Panel, 8.f, FStyleColors::White25, 2.f));
+
+	FButtonStyle CreateNewProjectButton = FButtonStyle()
+		.SetNormal(FSlateRoundedBoxBrush(FStyleColors::AccentBlue, 4.f))
+		.SetHovered(FSlateRoundedBoxBrush(FStyleColors::AccentBlue, 4.f))
+		.SetPressed(FSlateRoundedBoxBrush(FStyleColors::AccentBlue, 4.f))
+		.SetHoveredForeground(FStyleColors::Black)
+		.SetNormalForeground(FStyleColors::Black)
+		.SetPressedForeground(FStyleColors::Black)
+		.SetNormalPadding(FMargin(0.f, 4.0f))
+		.SetPressedPadding(FMargin(0.f, 4.0f));
+	Set("HomeScreen.CreateNewProjectButton", CreateNewProjectButton);
+
+	FButtonStyle MyFolderButton = FButtonStyle()
+		.SetNormal(FSlateRoundedBoxBrush(FStyleColors::DropdownOutline, 4.0f))
+		.SetHovered(FSlateRoundedBoxBrush(FStyleColors::Hover, 4.0f))
+		.SetPressed(FSlateRoundedBoxBrush(FStyleColors::DropdownOutline, 4.0f))
+		.SetHoveredForeground(FStyleColors::White)
+		.SetNormalForeground(FStyleColors::White)
+		.SetPressedForeground(FStyleColors::White)
+		.SetNormalPadding(FMargin(0.f, 4.0f))
+		.SetPressedPadding(FMargin(0.f, 4.0f));
+	Set("HomeScreen.MyFolderButton", MyFolderButton);
+
+	FSlateFontInfo RecentProjectFont = DEFAULT_FONT("Bold", 18);
+	Set("HomeScreen.RecentProjectFont", RecentProjectFont);
+
+	FButtonStyle SeeAllProjectsButton = FButtonStyle()
+		.SetNormal(FSlateRoundedBoxBrush(FStyleColors::Transparent, 4.0f, COLOR("FFFFFF59"), 1.f))
+		.SetHovered(FSlateRoundedBoxBrush(COLOR("FFFFFF0D"), 4.0f, COLOR("FFFFFF59"), 1.f))
+		.SetPressed(FSlateRoundedBoxBrush(COLOR("FFFFFF1A"), 4.0f, COLOR("FFFFFF59"), 1.f))
+		.SetHoveredForeground(FStyleColors::White)
+		.SetNormalForeground(FStyleColors::White)
+		.SetPressedForeground(FStyleColors::White)
+		.SetNormalPadding(FMargin(0.f, 4.0f))
+		.SetPressedPadding(FMargin(0.f, 4.0f));
+	Set("HomeScreen.SeeAllProjectsButton", SeeAllProjectsButton);
+
+	FButtonStyle ReconnectButton = FButtonStyle()
+		.SetNormal(FSlateRoundedBoxBrush(FStyleColors::DropdownOutline, 4.0f))
+		.SetHovered(FSlateRoundedBoxBrush(FStyleColors::Hover, 4.0f))
+		.SetPressed(FSlateRoundedBoxBrush(FStyleColors::DropdownOutline, 4.0f))
+		.SetHoveredForeground(FStyleColors::White)
+		.SetNormalForeground(FStyleColors::White)
+		.SetPressedForeground(FStyleColors::White)
+		.SetNormalPadding(FMargin(0.f, 4.0f))
+		.SetPressedPadding(FMargin(0.f, 4.0f));
+	Set("HomeScreen.ReconnectButton", ReconnectButton);
+
+	FCheckBoxStyle MainMenuSectionCheckBox = FAppStyle::GetWidgetStyle<FCheckBoxStyle>("ToggleButtonCheckbox");
+	MainMenuSectionCheckBox
+		.SetUncheckedHoveredImage(FSlateRoundedBoxBrush(FStyleColors::DropdownOutline, 4.0f))
+		.SetUncheckedPressedImage(FSlateRoundedBoxBrush(FStyleColors::DropdownOutline, 4.0f))
+		.SetCheckedImage(FSlateRoundedBoxBrush(FStyleColors::DropdownOutline, 4.0f))
+		.SetCheckedHoveredImage(FSlateRoundedBoxBrush(FStyleColors::Hover, 4.0f));
+	Set("HomeScreen.MainMenuSectionButton", MainMenuSectionCheckBox);
+}
+
 void FStarshipEditorStyle::FStyle::SetupTranslationEditorStyles()
 {
 	Set("TranslationEditor.Export", new IMAGE_BRUSH("Icons/Icon_Localisation_Export_All_40x", Icon40x40));
diff --git a/Engine/Source/Editor/EditorStyle/Private/StarshipStyle.h b/Engine/Source/Editor/EditorStyle/Private/StarshipStyle.h
--- a/Engine/Source/Editor/EditorStyle/Private/StarshipStyle.h
+++ b/Engine/Source/Editor/EditorStyle/Private/StarshipStyle.h
@@ -73,6 +73,7 @@ public:
 		void SetupSourceControlStyles();
 		void SetupAutomationStyles();
 		void SetupUMGEditorStyles();
+		void SetupHomeScreenStyles();
 		void SetupMyBlueprintStyles();
 		void SetupStatusBarStyle();
 		void SetupColorPickerStyle();
diff --git a/Engine/Source/Editor/GameProjectGeneration/Private/GameProjectGenerationModule.cpp b/Engine/Source/Editor/GameProjectGeneration/Private/GameProjectGenerationModule.cpp
--- a/Engine/Source/Editor/GameProjectGeneration/Private/GameProjectGenerationModule.cpp
+++ b/Engine/Source/Editor/GameProjectGeneration/Private/GameProjectGenerationModule.cpp
@@ -18,6 +18,7 @@
 #include "Misc/Paths.h"
 #include "Misc/ScopedSlowTask.h"
 #include "SNewClassDialog.h"
+#include "SProjectBrowser.h"
 #include "SProjectDialog.h"
 #include "Styling/SlateBrush.h"
 #include "TemplateCategory.h"
@@ -67,6 +68,11 @@ TSharedRef<class SWidget> FGameProjectGenerationModule::CreateNewClassDialog(con
 	return SNew(SNewClassDialog).Class(InClass);
 }
 
+TSharedRef<SWidget> FGameProjectGenerationModule::CreateProjectBrowser(bool bInShowHeaderBar, bool bInAllowScrolling)
+{
+	return SNew(SProjectBrowser).ShowHeaderBar(bInShowHeaderBar).AllowScrolling(bInAllowScrolling);
+}
+
 void FGameProjectGenerationModule::OpenAddCodeToProjectDialog(const FAddToProjectConfig& Config)
 {
 	GameProjectUtils::OpenAddToProjectDialog(Config, EClassDomain::Native);
diff --git a/Engine/Source/Editor/GameProjectGeneration/Private/SProjectBrowser.cpp b/Engine/Source/Editor/GameProjectGeneration/Private/SProjectBrowser.cpp
--- a/Engine/Source/Editor/GameProjectGeneration/Private/SProjectBrowser.cpp
+++ b/Engine/Source/Editor/GameProjectGeneration/Private/SProjectBrowser.cpp
@@ -378,6 +378,10 @@ void SProjectBrowser::Construct( const FArguments& InArgs )
 		.OnContextMenuOpening(this, &SProjectBrowser::OnGetContextMenuContent)
 		.OnMouseButtonDoubleClick(this, &SProjectBrowser::HandleProjectItemDoubleClick)
 		.OnSelectionChanged(this, &SProjectBrowser::HandleProjectViewSelectionChanged)
+		.AllowOverscroll(InArgs._AllowScrolling ? EAllowOverscroll::Yes : EAllowOverscroll::No)
+		.ConsumeMouseWheel(InArgs._AllowScrolling ? EConsumeMouseWheel::Always : EConsumeMouseWheel::Never)
+		.ScrollbarVisibility(InArgs._AllowScrolling ? EVisibility::Visible : EVisibility::Collapsed)
+		.WheelScrollMultiplier(InArgs._AllowScrolling ? GetGlobalScrollAmount() : 0.f)
 		.ItemHeight(ProjectBrowserDefs::ProjectTileHeight+9)
 		.ItemWidth(ProjectBrowserDefs::ProjectTileWidth+9);
 
@@ -430,7 +434,9 @@ void SProjectBrowser::Construct( const FArguments& InArgs )
 		+ SVerticalBox::Slot()
 		.AutoHeight()
 		[
-			SNew(SHorizontalBox)	
+			SNew(SHorizontalBox)
+			.Visibility(InArgs._ShowHeaderBar ? EVisibility::Visible : EVisibility::Collapsed)
+
 			+ SHorizontalBox::Slot()
 			.Padding(FMargin(0, 0, 5.f, 5))
 			.VAlign(VAlign_Center)
diff --git a/Engine/Source/Editor/GameProjectGeneration/Private/SProjectBrowser.h b/Engine/Source/Editor/GameProjectGeneration/Private/SProjectBrowser.h
--- a/Engine/Source/Editor/GameProjectGeneration/Private/SProjectBrowser.h
+++ b/Engine/Source/Editor/GameProjectGeneration/Private/SProjectBrowser.h
@@ -62,8 +62,12 @@ public:
 	DECLARE_DELEGATE(FNewProjectScreenRequested)
 
 	SLATE_BEGIN_ARGS(SProjectBrowser)
+		: _ShowHeaderBar(true)
+		, _AllowScrolling(true)
 	{}
-	SLATE_ARGUMENT(FProjectSelectionChanged, OnSelectionChanged);
+		SLATE_ARGUMENT(bool, ShowHeaderBar)
+		SLATE_ARGUMENT(bool, AllowScrolling)
+		SLATE_ARGUMENT(FProjectSelectionChanged, OnSelectionChanged);
 	SLATE_END_ARGS()
 
 public:
diff --git a/Engine/Source/Editor/GameProjectGeneration/Public/GameProjectGenerationModule.h b/Engine/Source/Editor/GameProjectGeneration/Public/GameProjectGenerationModule.h
--- a/Engine/Source/Editor/GameProjectGeneration/Public/GameProjectGenerationModule.h
+++ b/Engine/Source/Editor/GameProjectGeneration/Public/GameProjectGenerationModule.h
@@ -58,7 +58,10 @@ public:
 
 	/** Creates a new class dialog for creating classes based on the passed-in class. */
 	virtual TSharedRef<class SWidget> CreateNewClassDialog(const UClass* InClass);
-	
+
+	/** Creates the project browser widget */
+	virtual TSharedRef<class SWidget> CreateProjectBrowser(bool bInShowHeaderBar = true, bool bInAllowScrolling = true);
+
 	/** 
 	 * Opens a dialog to add code files to the current project. 
 	 *
diff --git a/Engine/Source/Editor/MainFrame/MainFrame.Build.cs b/Engine/Source/Editor/MainFrame/MainFrame.Build.cs
--- a/Engine/Source/Editor/MainFrame/MainFrame.Build.cs
+++ b/Engine/Source/Editor/MainFrame/MainFrame.Build.cs
@@ -21,6 +21,7 @@ public class MainFrame : ModuleRules
 				"ApplicationCore",
 				"Engine",
 				"EngineSettings",
+				"HTTP",
 				"InputCore",
 				"RHI",
 				"RenderCore",
@@ -45,6 +46,7 @@ public class MainFrame : ModuleRules
 				"InterchangeCore",
 				"InterchangeEngine",
 				"ToolWidgets",
+				"WebBrowser",
 			}
 		);
 
diff --git a/Engine/Source/Editor/MainFrame/Private/HomeScreen/HomeScreenCommon.h b/Engine/Source/Editor/MainFrame/Private/HomeScreen/HomeScreenCommon.h
new file mode 100644
--- /dev/null
+++ b/Engine/Source/Editor/MainFrame/Private/HomeScreen/HomeScreenCommon.h
@@ -0,0 +1,14 @@
+﻿// Copyright Epic Games, Inc. All Rights Reserved.
+
+#pragma once
+#include "HomeScreenCommon.generated.h"
+
+UENUM()
+enum class EMainSectionMenu : uint8
+{
+	None,
+	Home,
+	News,
+	GettingStarted,
+	SampleProjects
+};
diff --git a/Engine/Source/Editor/MainFrame/Private/HomeScreen/HomeScreenWeb.cpp b/Engine/Source/Editor/MainFrame/Private/HomeScreen/HomeScreenWeb.cpp
new file mode 100644
--- /dev/null
+++ b/Engine/Source/Editor/MainFrame/Private/HomeScreen/HomeScreenWeb.cpp
@@ -0,0 +1,15 @@
+﻿// Copyright Epic Games, Inc. All Rights Reserved.
+
+#include "HomeScreenWeb.h"
+
+void UHomeScreenWeb::NavigateTo(EMainSectionMenu InSectionToNavigate)
+{
+	SectionToNavigate = InSectionToNavigate;
+
+	OnNavigationChangedDelegate.Broadcast(InSectionToNavigate);
+}
+
+void UHomeScreenWeb::OpenWebPage(const FString& InURL) const
+{
+	FPlatformProcess::LaunchURL(*InURL, nullptr, nullptr);
+}
diff --git a/Engine/Source/Editor/MainFrame/Private/HomeScreen/HomeScreenWeb.h b/Engine/Source/Editor/MainFrame/Private/HomeScreen/HomeScreenWeb.h
new file mode 100644
--- /dev/null
+++ b/Engine/Source/Editor/MainFrame/Private/HomeScreen/HomeScreenWeb.h
@@ -0,0 +1,26 @@
+﻿// Copyright Epic Games, Inc. All Rights Reserved.
+
+#pragma once
+#include "HomeScreenCommon.h"
+#include "UObject/Object.h"
+#include "HomeScreenWeb.generated.h"
+
+DECLARE_MULTICAST_DELEGATE_OneParam(FOnNavigationChanged, const EMainSectionMenu);
+
+UCLASS()
+class UHomeScreenWeb : public UObject
+{
+	GENERATED_BODY()
+
+public:
+	UFUNCTION()
+	void NavigateTo(EMainSectionMenu InSectionToNavigate);
+	FOnNavigationChanged& OnNavigationChanged() { return OnNavigationChangedDelegate; }
+
+	UFUNCTION()
+	void OpenWebPage(const FString& InURL) const;
+
+private:
+	FOnNavigationChanged OnNavigationChangedDelegate;
+	EMainSectionMenu SectionToNavigate = EMainSectionMenu::None;
+};
diff --git a/Engine/Source/Editor/MainFrame/Private/HomeScreen/SHomeScreen.cpp b/Engine/Source/Editor/MainFrame/Private/HomeScreen/SHomeScreen.cpp
new file mode 100644
--- /dev/null
+++ b/Engine/Source/Editor/MainFrame/Private/HomeScreen/SHomeScreen.cpp
@@ -0,0 +1,965 @@
+﻿// Copyright Epic Games, Inc. All Rights Reserved.
+
+#include "SHomeScreen.h"
+
+#include "DesktopPlatformModule.h"
+#include "Editor.h"
+#include "Frame/RootWindowLocation.h"
+#include "GameProjectGenerationModule.h"
+#include "GameProjectUtils.h"
+#include "HomeScreenWeb.h"
+#include "HttpModule.h"
+#include "HttpRetrySystem.h"
+#include "IWebBrowserWindow.h"
+#include "Misc/StringOutputDevice.h"
+#include "Settings/EditorSettings.h"
+#include "Styling/AppStyle.h"
+#include "Styling/StyleColors.h"
+#include "SWebBrowser.h"
+#include "TimerManager.h"
+#include "Widgets/Images/SImage.h"
+#include "Widgets/Images/SThrobber.h"
+#include "Widgets/Input/SButton.h"
+#include "Widgets/Input/SCheckBox.h"
+#include "Widgets/Input/SComboButton.h"
+#include "Widgets/Layout/SBorder.h"
+#include "Widgets/Layout/SBox.h"
+#include "Widgets/Layout/SSeparator.h"
+#include "Widgets/Layout/SSpacer.h"
+#include "Widgets/Layout/SWidgetSwitcher.h"
+#include "Widgets/SBoxPanel.h"
+#include "Widgets/Text/STextBlock.h"
+#include "Widgets/SWidget.h"
+
+#define LOCTEXT_NAMESPACE "SHomeScreen"
+
+SHomeScreen::SHomeScreen()
+{
+	if (!HttpRetryManager.IsValid())
+	{
+		HttpRetryManager = MakeShared<FHttpRetrySystem::FManager>(
+			FHttpRetrySystem::FRetryLimitCountSetting(3),
+			FHttpRetrySystem::FRetryTimeoutRelativeSecondsSetting(3));
+	}
+
+	if (GEditor && GEditor->IsTimerManagerValid())
+	{
+		constexpr float TimerManagerCallRate = 5.f;
+		constexpr bool bLoop = true;
+
+		GEditor->GetTimerManager()->SetTimer(
+			CheckInternetConnectionTimerHandle,
+			FTimerDelegate::CreateLambda([this] ()
+				{
+					if (bIsRequestFinished)
+					{
+						if (bIsConnected || bForceRetry)
+						{
+							bForceRetry = false;
+							TimerManagerCountOnceDisconnected = 0;
+							CheckInternetConnection();
+						}
+						else
+						{
+							// When disconnected retry connection less times than when you are connected
+							TimerManagerCountOnceDisconnected += 1;
+							bForceRetry = TimerManagerCountOnceDisconnected == MaxTimerManagerCountOnceDisconnectedBeforeRetry;
+						}
+					}
+				}),
+				TimerManagerCallRate,
+				bLoop);
+	}
+}
+
+SHomeScreen::~SHomeScreen()
+{
+	SaveSettings();
+
+	if (GEditor && GEditor->IsTimerManagerValid())
+	{
+		GEditor->GetTimerManager()->ClearTimer(CheckInternetConnectionTimerHandle);
+	}
+
+	if (WebObject.IsValid())
+	{
+		WebObject->OnNavigationChanged().RemoveAll(this);
+	}
+}
+
+void SHomeScreen::Construct(const FArguments& InArgs, const TSharedPtr<SWebBrowser> InWebBrowser, const TWeakObjectPtr<UHomeScreenWeb> InWebObject)
+{
+	WebBrowser = InWebBrowser;
+	WebObject = InWebObject;
+	if (WebObject.IsValid())
+	{
+		WebObject->OnNavigationChanged().AddSP(this, &SHomeScreen::OnNavigateToSection);
+	}
+
+	AutoLoadProjectComboBoxSelection = EAutoLoadProject::HomeScreen;
+	MainHomeSelection = EMainSectionMenu::Home;
+
+	LoadSettings();
+
+	// Main Section CheckBoxes
+	CreateMainSectionCheckBox(HomeCheckBox, EMainSectionMenu::Home, LOCTEXT("ProjectHomeHome", "Home"), FAppStyle::GetBrush("HomeScreen.Home"));
+	CreateMainSectionCheckBox(NewsCheckBox, EMainSectionMenu::News, LOCTEXT("ProjectHomeNews", "News"), FAppStyle::GetBrush("HomeScreen.News"));
+	CreateMainSectionCheckBox(GettingStartedCheckBox, EMainSectionMenu::GettingStarted, LOCTEXT("ProjectHomeGettingStarted", "Getting Started"), FAppStyle::GetBrush("HomeScreen.Rocket"));
+	CreateMainSectionCheckBox(SampleProjectsCheckBox, EMainSectionMenu::SampleProjects, LOCTEXT("ProjectHomeSampleProjects", "Sample Projects"), FAppStyle::GetBrush("HomeScreen.Archive"));
+
+	// Documentation and Community Buttons
+	CreateResourceButtons(ForumsButton, TEXT("https://forums.unrealengine.com/categories?tag=unreal-engine"),
+		LOCTEXT("ProjectHomeForumsButton", "Forums"), FAppStyle::GetBrush("HomeScreen.Forum"));
+	CreateResourceButtons(DocumentationButton, TEXT("https://dev.epicgames.com/documentation/en-us/unreal-engine/unreal-engine-5-6-documentation"),
+		LOCTEXT("ProjectHomeDocumentation", "Documentation"), FAppStyle::GetBrush("HomeScreen.Documentation"));
+	CreateResourceButtons(TutorialsButton, TEXT("https://dev.epicgames.com/community/unreal-engine/learning"),
+		LOCTEXT("ProjectHomeTutorialsButton", "Tutorials"), FAppStyle::GetBrush("HomeScreen.Tutorial"));
+	CreateResourceButtons(RoadmapButton, TEXT("https://portal.productboard.com/epicgames/1-unreal-engine-public-roadmap/tabs/109-unreal-engine-5-5"),
+		LOCTEXT("ProjectHomeRoadmapButton", "Roadmap"), FAppStyle::GetBrush("HomeScreen.Roadmap"));
+	CreateResourceButtons(ReleaseNotesButton, TEXT("https://dev.epicgames.com/documentation/en-us/unreal-engine/unreal-engine-5-6-release-notes"),
+		LOCTEXT("ProjectHomeReleaseNotesButton", "Release Notes"), FAppStyle::GetBrush("HomeScreen.ReleaseNotes"));
+
+	// Social Media Buttons
+	TSharedPtr<SButton> FacebookButton;
+	TSharedPtr<SButton> TwitterButton;
+	TSharedPtr<SButton> YouTubeButton;
+	TSharedPtr<SButton> TwitchButton;
+	TSharedPtr<SButton> InstagramButton;
+
+	CreateSocialMediaButtons(FacebookButton, TEXT("https://www.facebook.com/unrealengine"), FAppStyle::GetBrush("SocialMedia.Facebook"));
+	CreateSocialMediaButtons(TwitterButton, TEXT("https://twitter.com/unrealengine"), FAppStyle::GetBrush("SocialMedia.Twitter"));
+	CreateSocialMediaButtons(YouTubeButton, TEXT("https://www.youtube.com/c/unrealengine"), FAppStyle::GetBrush("SocialMedia.YouTube"));
+	CreateSocialMediaButtons(TwitchButton, TEXT("https://www.twitch.tv/unrealengine"), FAppStyle::GetBrush("SocialMedia.Twitch"));
+	CreateSocialMediaButtons(InstagramButton, TEXT("https://www.instagram.com/unrealengine/"), FAppStyle::GetBrush("SocialMedia.Instagram"));
+
+	// Project Browser
+	ProjectBrowser = FGameProjectGenerationModule::Get().CreateProjectBrowser(false, false);
+
+	// Check for internet connection on startup
+	CheckInternetConnection();
+
+	ChildSlot
+	[
+		SNew(SBorder)
+		.BorderImage(FAppStyle::GetBrush("HomeScreen.Background"))
+		.Padding(0.f)
+		[
+			SNew(SHorizontalBox)
+
+			// New Project Section
+			+ SHorizontalBox::Slot()
+			.AutoWidth()
+			[
+				SNew(SBorder)
+				.BorderImage(FAppStyle::GetBrush("ToolPanel.GroupBorder"))
+				.Padding(8.f)
+				[
+					// Title Icon
+					SNew(SVerticalBox)
+
+					+ SVerticalBox::Slot()
+					.AutoHeight()
+					.HAlign(HAlign_Center)
+					.Padding(0.f, 32.f)
+					[
+						SNew(SImage)
+						.Image(FAppStyle::GetBrush("HomeScreen.UnrealLogo"))
+					]
+
+					// Create and MyFolder Buttons
+					+ SVerticalBox::Slot()
+					.AutoHeight()
+					.Padding(8.f, 0.f, 8.f, 8.f)
+					[
+						SNew(SButton)
+						.ButtonStyle(FAppStyle::Get(), "HomeScreen.CreateNewProjectButton")
+						.OnClicked(this, &SHomeScreen::OnCreateProjectDialog, /** bInAllowProjectOpening */false, /** bInAllowProjectCreation */true)
+						.HAlign(HAlign_Center)
+						[
+							SNew(SHorizontalBox)
+
+							+ SHorizontalBox::Slot()
+							.AutoWidth()
+							.Padding(0.f, 0.f, 8.f, 0.f)
+							[
+								SNew(SImage)
+								.Image(FAppStyle::GetBrush("HomeScreen.PlusCircle"))
+								.ColorAndOpacity(FStyleColors::Black)
+							]
+
+							+ SHorizontalBox::Slot()
+							.AutoWidth()
+							.VAlign(VAlign_Center)
+							[
+								SNew(STextBlock)
+								.Text(LOCTEXT("ProjectHomeCreateNewProject", "New Project"))
+							]
+						]
+					]
+
+					+ SVerticalBox::Slot()
+					.AutoHeight()
+					.Padding(8.f, 0.f, 8.f, 0.f)
+					[
+						SNew(SButton)
+						.ButtonStyle(FAppStyle::Get(), "HomeScreen.MyFolderButton")
+						.OnClicked(this, &SHomeScreen::OnCreateProjectDialog, /** bInAllowProjectOpening */true, /** bInAllowProjectCreation */false)
+						.HAlign(HAlign_Center)
+						[
+							SNew(SHorizontalBox)
+
+							+ SHorizontalBox::Slot()
+							.AutoWidth()
+							.Padding(0.f, 0.f, 8.f, 0.f)
+							[
+								SNew(SImage)
+								.Image(FAppStyle::GetBrush("HomeScreen.FolderOpen"))
+							]
+
+							+ SHorizontalBox::Slot()
+							.AutoWidth()
+							.VAlign(VAlign_Center)
+							[
+								SNew(STextBlock)
+								.Text(LOCTEXT("ProjectHomeMyProjects", "My Projects"))
+							]
+						]
+					]
+
+					// Main Home Section
+					+ SVerticalBox::Slot()
+					.AutoHeight()
+					.Padding(8.f, 16.f)
+					[
+						SNew(SVerticalBox)
+
+						+ SVerticalBox::Slot()
+						.AutoHeight()
+						[
+							HomeCheckBox.ToSharedRef()
+						]
+					
+						+ SVerticalBox::Slot()
+						.AutoHeight()
+						[
+							NewsCheckBox.ToSharedRef()
+						]
+
+						+ SVerticalBox::Slot()
+						.AutoHeight()
+						[
+							GettingStartedCheckBox.ToSharedRef()
+						]
+
+						+ SVerticalBox::Slot()
+						.AutoHeight()
+						[
+							SampleProjectsCheckBox.ToSharedRef()
+						]
+					]
+
+					// Documentation and Community Section
+					+ SVerticalBox::Slot()
+					.AutoHeight()
+					.Padding(8.f, 0.f, 8.f, 4.f)
+					[
+						SNew(STextBlock)
+						.Text(LOCTEXT("ProjectHomeDocumentation", "Resources & Community"))
+						.Font(FAppStyle::GetFontStyle("BoldFont"))
+						.ColorAndOpacity(FStyleColors::White)
+					]
+
+					+ SVerticalBox::Slot()
+					.AutoHeight()
+					.Padding(16.f, 0.f)
+					[
+						SNew(SVerticalBox)
+
+						+ SVerticalBox::Slot()
+						.AutoHeight()
+						[
+							ForumsButton.ToSharedRef()
+						]
+
+						+ SVerticalBox::Slot()
+						.AutoHeight()
+						[
+							DocumentationButton.ToSharedRef()
+						]
+
+						+ SVerticalBox::Slot()
+						.AutoHeight()
+						[
+							TutorialsButton.ToSharedRef()
+						]
+
+						+ SVerticalBox::Slot()
+						.AutoHeight()
+						[
+							RoadmapButton.ToSharedRef()
+						]
+
+						+ SVerticalBox::Slot()
+						.AutoHeight()
+						[
+							ReleaseNotesButton.ToSharedRef()
+						]
+					]
+
+					// Load on Startup Section
+					+ SVerticalBox::Slot()
+					.VAlign(VAlign_Bottom)
+					.FillHeight(1.f)
+					[
+						SNew(SVerticalBox)
+
+						+ SVerticalBox::Slot()
+						.Padding(16.f, 0.f, 16.f, 16.f)
+						[
+							SNew(SBorder)
+							.Clipping(EWidgetClipping::ClipToBounds)
+							.BorderImage(FAppStyle::GetBrush("HomeScreen.NoInternet.Border"))
+							.Padding(16.f)
+							.Visibility_Lambda([this]()
+							{
+								return bIsConnected ? EVisibility::Collapsed : EVisibility::Visible;
+							})
+							[
+								SNew(SVerticalBox)
+
+								+ SVerticalBox::Slot()
+								.AutoHeight()
+								.Padding(0.f, 0.f, 0.f, 16.f)
+								.HAlign(HAlign_Left)
+								[
+									SNew(SBox)
+									.WidthOverride(24.f)
+									.HeightOverride(24.f)
+									[
+										SNew(SWidgetSwitcher)
+										.WidgetIndex(this, &SHomeScreen::GetNoInternetIconIndex)
+
+										+ SWidgetSwitcher::Slot()
+										[
+											SNew(SImage)
+											.Image(FAppStyle::GetBrush("HomeScreen.NoInternet.Icon"))
+										]
+
+										+ SWidgetSwitcher::Slot()
+										[
+											SNew(SCircularThrobber)
+											.NumPieces(8)
+											.Period(1.2f)
+											.Radius(12.f)
+										]
+									]
+								]
+
+								+ SVerticalBox::Slot()
+								.FillHeight(0.5f)
+								.Padding(0.f, 0.f, 0.f, 16.f)
+								[
+									SNew(STextBlock)
+									.Font(FAppStyle::GetFontStyle("BoldFont"))
+									.Text(LOCTEXT("HomeScreenNoInternet", "No Internet Connection"))
+									.ColorAndOpacity(FStyleColors::White)
+								]
+
+								+ SVerticalBox::Slot()
+								.FillHeight(1.f)
+								[
+									SNew(SButton)
+									.HAlign(HAlign_Center)
+									.ButtonStyle(FAppStyle::Get(), "HomeScreen.MyFolderButton")
+									.Text(LOCTEXT("HomeScreenReconnect", "Reconnect"))
+									.OnClicked(this, &SHomeScreen::OnInternetConnectionRetried)
+								]
+							]
+						]
+
+						+ SVerticalBox::Slot()
+						.AutoHeight()
+						.Padding(16.f, 0.f)
+						.HAlign(HAlign_Left)
+						[
+							SNew(STextBlock)
+							.Text(LOCTEXT("ProjectHomeLoad", "Load on Startup"))
+						]
+
+						+ SVerticalBox::Slot()
+						.AutoHeight()
+						.Padding(16.f, 8.f, 16.f, 36.f)
+						[
+							SAssignNew(AutoLoadProjectComboBox, SComboButton)
+							.ComboButtonStyle(&FAppStyle::Get().GetWidgetStyle<FComboButtonStyle>("HomeScreen.ComboButton"))
+							.MenuPlacement(EMenuPlacement::MenuPlacement_CenteredAboveAnchor)
+							.HasDownArrow(true)
+							.ButtonContent()
+							[
+								SNew(STextBlock)
+								.Text(this, &SHomeScreen::GetAutoLoadProjectComboBoxLabelText)
+								.ColorAndOpacity_Lambda([this] ()
+								{
+									if (AutoLoadProjectComboBox.IsValid() && AutoLoadProjectComboBox->IsHovered())
+									{
+										return FStyleColors::White;
+									}
+									return FStyleColors::Foreground;
+								})
+							]
+							.MenuContent()
+							[
+								CreateComboButtonMenuContentWidget()
+							]
+						]
+
+						// Social Media Section
+						+ SVerticalBox::Slot()
+						.AutoHeight()
+						.Padding(46.f, 0.f)
+						[
+							SNew(SHorizontalBox)
+
+							+ SHorizontalBox::Slot()
+							.AutoWidth()
+							.Padding(4.f, 0.f)
+							[
+								FacebookButton.ToSharedRef()
+							]
+
+							+ SHorizontalBox::Slot()
+							.AutoWidth()
+							.Padding(4.f, 0.f)
+							[
+								TwitterButton.ToSharedRef()
+							]
+
+							+ SHorizontalBox::Slot()
+							.AutoWidth()
+							.Padding(4.f, 0.f)
+							[
+								YouTubeButton.ToSharedRef()
+							]
+
+							+ SHorizontalBox::Slot()
+							.AutoWidth()
+							.Padding(4.f, 0.f)
+							[
+								TwitchButton.ToSharedRef()
+							]
+
+							+ SHorizontalBox::Slot()
+							.AutoWidth()
+							.Padding(4.f, 0.f)
+							[
+								InstagramButton.ToSharedRef()
+							]
+						]
+					]
+				]
+			]
+
+			// Web Api
+			+ SHorizontalBox::Slot()
+			.FillWidth(1.f)
+			[
+				SNew(SWidgetSwitcher)
+				.WidgetIndex_Lambda([this] ()
+				{
+					return MainHomeSelection == EMainSectionMenu::Home ? 0 : 1;
+				})
+
+				+ SWidgetSwitcher::Slot()
+				[
+					SNew(SVerticalBox)
+
+					+ SVerticalBox::Slot()
+					.AutoHeight()
+					[
+						SNew(SHorizontalBox)
+
+						+ SHorizontalBox::Slot()
+						.FillWidth(0.125f)
+						[
+							SNew(SSpacer)
+						]
+
+						+ SHorizontalBox::Slot()
+						.FillWidth(0.75f)
+						[
+							SNew(SVerticalBox)
+							+ SVerticalBox::Slot()
+							.Padding(0.f, 64.f, 0.f, 16.f)
+							[
+								SNew(SHorizontalBox)
+							
+								+ SHorizontalBox::Slot()
+								.AutoWidth()
+								.VAlign(VAlign_Center)
+								.Padding(0.f, 0.f, 8.f, 0.f)
+								[
+									SNew(STextBlock)
+									.Text(LOCTEXT("RecentProjects", "Recent Projects"))
+									.Font(FAppStyle::GetFontStyle("HomeScreen.RecentProjectFont"))
+									.ColorAndOpacity(FStyleColors::White)
+								]
+							
+								+ SHorizontalBox::Slot()
+								.AutoWidth()
+								.VAlign(VAlign_Center)
+								[
+									SNew(SButton)
+									.ButtonStyle(FAppStyle::Get(), "HomeScreen.SeeAllProjectsButton")
+									.OnClicked(this, &SHomeScreen::OnCreateProjectDialog, /** bInAllowProjectOpening */true, /** bInAllowProjectCreation */false)
+									.Text(LOCTEXT("SeeAllProjects", "See All Projects"))
+								]
+							]
+
+							+ SVerticalBox::Slot()
+							.AutoHeight()
+							[
+								SNew(SBorder)
+								.Clipping(EWidgetClipping::ClipToBounds)
+								.BorderImage(FAppStyle::GetBrush("HomeScreen.RecentProjects.Background"))
+								.Padding(16.f)
+								[
+									SNew(SVerticalBox)
+
+									+ SVerticalBox::Slot()
+									[
+										SNew(SBox)
+										.HeightOverride(172.f)
+										.Padding(0.f)
+										[
+												ProjectBrowser.ToSharedRef()
+										]
+									]
+								]
+							]
+						]
+
+						+ SHorizontalBox::Slot()
+						.FillWidth(0.125f)
+						[
+							SNew(SSpacer)
+						]
+					]
+
+					+ SVerticalBox::Slot()
+					.FillHeight(1.f)
+					[
+						WebBrowser.ToSharedRef()
+					]
+				]
+
+				+ SWidgetSwitcher::Slot()
+				[
+					WebBrowser.ToSharedRef()
+				]
+			]
+		]
+	];
+
+	OnMainHomeSectionChanged(ECheckBoxState::Checked, MainHomeSelection);
+}
+
+FReply SHomeScreen::OnCreateProjectDialog(bool bInAllowProjectOpening, bool bInAllowProjectCreation)
+{
+	FRootWindowLocation DefaultWindowLocation;
+
+	FText Title = LOCTEXT("ProjectBrowserTitle", "Unreal Project Browser");
+	FText ToolTip = LOCTEXT("ProjectBrowserTitleToolTip", "Close Unreal Project Browser");
+	if (bInAllowProjectCreation)
+	{
+		Title = LOCTEXT("ProjectCreationTitle", "Unreal Project Creation");
+		ToolTip = LOCTEXT("ProjectCreationTitleToolTip", "Close Unreal Project Creation");
+	}
+
+	const TSharedRef<SWindow> NewProjectWindow = SNew(SWindow)
+		.AutoCenter(EAutoCenter::PreferredWorkArea)
+		.Title(Title)
+		.CloseButtonToolTipText(ToolTip)
+		.IsInitiallyMaximized(false)
+		.ScreenPosition(DefaultWindowLocation.ScreenPosition)
+		.ClientSize(DefaultWindowLocation.WindowSize)
+		.CreateTitleBar(true)
+		.SizingRule(ESizingRule::UserSized)
+		.SupportsMaximize(false)
+		.SupportsMinimize(false);
+
+	NewProjectWindow->SetContent(FGameProjectGenerationModule::Get().CreateGameProjectDialog(bInAllowProjectOpening, bInAllowProjectCreation));
+	FSlateApplication::Get().AddModalWindow(NewProjectWindow, nullptr);
+	
+	return FReply::Handled();
+}
+
+FReply SHomeScreen::OnSocialMediaClicked(FString InURL)
+{
+	FPlatformProcess::LaunchURL(*InURL, nullptr, nullptr);
+	return FReply::Handled();
+}
+
+ECheckBoxState SHomeScreen::IsMainHomeSectionChecked(EMainSectionMenu InMainHomeSelection) const
+{
+	return MainHomeSelection == InMainHomeSelection ? ECheckBoxState::Checked : ECheckBoxState::Unchecked;
+}
+
+void SHomeScreen::OnMainHomeSectionChanged(ECheckBoxState InCheckBoxState, EMainSectionMenu InMainHomeSelection)
+{
+	FString URL = TEXT("");
+	switch (InMainHomeSelection)
+	{
+	case EMainSectionMenu::News:
+		MainHomeSelection = EMainSectionMenu::News;
+		URL = TEXT("https://www.unrealengine.com/en-US/feed");
+		break;
+	case EMainSectionMenu::GettingStarted:
+		MainHomeSelection = EMainSectionMenu::GettingStarted;
+		URL = TEXT("https://dev.epicgames.com/community/unreal-engine/getting-started");
+		break;
+	case EMainSectionMenu::SampleProjects:
+		MainHomeSelection = EMainSectionMenu::SampleProjects;
+		URL = TEXT("https://dev.epicgames.com/community/unreal-engine/learning");
+		break;
+	case EMainSectionMenu::None:
+	case EMainSectionMenu::Home:
+	default:
+		MainHomeSelection = EMainSectionMenu::Home;
+		URL = TEXT("https://editor-website.ol.epicgames.net/en-US");
+		break;
+	}
+
+	if (WebBrowser.IsValid())
+	{
+		WebBrowser->LoadURL(URL);
+	}
+}
+
+void SHomeScreen::LoadSettings()
+{
+	int32 LoadingType;
+	if (GConfig->GetInt(TEXT("HomeScreen"), TEXT("Loading"), LoadingType, GEngineIni))
+	{
+		// Clamp value to normal range and update loading type
+		LoadingType = FMath::Clamp<int32>(LoadingType, 0, (int32)EAutoLoadProject::MAX-1);
+
+		AutoLoadProjectComboBoxSelection = (EAutoLoadProject)LoadingType;
+	}
+
+	bool bShowHomePage;
+	if (GConfig->GetBool(TEXT("HomeScreen"), TEXT("ShowHomePage"), bShowHomePage, GEngineIni))
+	{
+		bShowHomePageOption = bShowHomePage;
+	}
+}
+
+void SHomeScreen::SaveSettings()
+{
+	GConfig->SetInt(TEXT("HomeScreen"), TEXT("Loading"), (int32)AutoLoadProjectComboBoxSelection, GEngineIni);
+	GConfig->SetBool(TEXT("HomeScreen"), TEXT("ShowHomePage"), bShowHomePageOption, GEngineIni);
+}
+
+FReply SHomeScreen::OnAutoloadLastProjectChanged(EAutoLoadProject InAutoLoadOption)
+{
+	UEditorSettings *Settings = GetMutableDefault<UEditorSettings>();
+	AutoLoadProjectComboBoxSelection = InAutoLoadOption;
+
+	Settings->bLoadTheMostRecentlyLoadedProjectAtStartup = AutoLoadProjectComboBoxSelection == EAutoLoadProject::LastProject;
+
+	// Manually trigger the property changed event so that the editor settings UI updates to reflect the change
+	FProperty* AutoloadProjectProperty = FindFProperty<FProperty>(Settings->GetClass(), "bLoadTheMostRecentlyLoadedProjectAtStartup");
+	if (AutoloadProjectProperty != NULL)
+	{
+		FPropertyChangedEvent PropertyUpdateStruct(AutoloadProjectProperty);
+		Settings->PostEditChangeProperty(PropertyUpdateStruct);
+	}
+
+	SaveSettings();
+	return FReply::Handled();
+}
+
+FReply SHomeScreen::OnShowHomePageWhenNewAnnouncementOptionClicked()
+{
+	bShowHomePageOption = !bShowHomePageOption;
+	SaveSettings();
+	return FReply::Handled();
+}
+
+bool SHomeScreen::IsCheckBoxCheckedOrHovered(const TSharedPtr<SCheckBox> InCheckBox) const
+{
+	return InCheckBox.IsValid() ? InCheckBox->IsHovered() || InCheckBox->IsChecked() : false;
+}
+
+FSlateColor SHomeScreen::GetMainSectionCheckBoxColor(const TSharedPtr<SCheckBox> InCheckBox) const
+{
+	return IsCheckBoxCheckedOrHovered(InCheckBox) ? FLinearColor::White : FLinearColor(1.f, 1.f, 1.f, 0.65f);
+}
+
+FSlateColor SHomeScreen::GetResourceAndSocialMediaButtonColor(const TSharedPtr<SButton> InButton) const
+{
+	if (InButton.IsValid() && InButton->IsHovered())
+	{
+		return FLinearColor::White;
+	}
+
+	return FLinearColor(1.f, 1.f, 1.f, 0.65f);
+}
+
+void SHomeScreen::CreateMainSectionCheckBox(TSharedPtr<SCheckBox>& OutCheckBox, EMainSectionMenu InMainHomeSelection, const FText& InText, const FSlateBrush* InImage)
+{
+	OutCheckBox = SNew(SCheckBox)
+		.IsEnabled(this, &SHomeScreen::IsMainSectionEnabled, InMainHomeSelection)
+		.Style(FAppStyle::Get(), "HomeScreen.MainMenuSectionButton")
+		.IsChecked(this, &SHomeScreen::IsMainHomeSectionChecked, InMainHomeSelection)
+		.OnCheckStateChanged(this, &SHomeScreen::OnMainHomeSectionChanged, InMainHomeSelection)
+		.Padding(FMargin(8.f, 2.f));
+
+	OutCheckBox->SetContent(
+		SNew(SHorizontalBox)
+
+		+ SHorizontalBox::Slot()
+		.AutoWidth()
+		.Padding(0.f, 0.f, 4.f, 0.f)
+		[
+			SNew(SImage)
+			.Image(InImage)
+			.ColorAndOpacity(this, &SHomeScreen::GetMainSectionCheckBoxColor, OutCheckBox)
+		]
+
+		+ SHorizontalBox::Slot()
+		.AutoWidth()
+		.VAlign(VAlign_Center)
+		[
+			SNew(STextBlock)
+			.Text(InText)
+			.ColorAndOpacity(this, &SHomeScreen::GetMainSectionCheckBoxColor, OutCheckBox)
+		]);
+}
+
+void SHomeScreen::CreateResourceButtons(TSharedPtr<SButton>& OutButton, FString InLink, const FText& InText, const FSlateBrush* InImage)
+{
+	SAssignNew(OutButton, SButton)
+	.IsEnabled(this, &SHomeScreen::IsConnectedToInternet)
+	.ButtonStyle(FAppStyle::Get(), "NoBorder")
+	.ContentPadding(FMargin(0.f, 2.f))
+	.OnClicked(this, &SHomeScreen::OnSocialMediaClicked, InLink);
+
+	OutButton->SetContent(
+		SNew(SHorizontalBox)
+
+		+ SHorizontalBox::Slot()
+		.AutoWidth()
+		.Padding(0.f, 0.f, 4.f, 0.f)
+		[
+			SNew(SImage)
+			.Image(InImage)
+			.ColorAndOpacity(this, &SHomeScreen::GetResourceAndSocialMediaButtonColor, OutButton)
+		]
+
+		+ SHorizontalBox::Slot()
+		.AutoWidth()
+		.VAlign(VAlign_Center)
+		[
+			SNew(STextBlock)
+			.Text(InText)
+			.ColorAndOpacity(this, &SHomeScreen::GetResourceAndSocialMediaButtonColor, OutButton)
+		]);
+}
+
+void SHomeScreen::CreateSocialMediaButtons(TSharedPtr<SButton>& OutButton, FString InLink, const FSlateBrush* InImage)
+{
+	SAssignNew(OutButton, SButton)
+	.ButtonStyle(FAppStyle::Get(), "NoBorder")
+	.OnClicked(this, &SHomeScreen::OnSocialMediaClicked, InLink);
+
+	OutButton->SetContent(
+			SNew(SImage)
+			.Image(InImage)
+			.ColorAndOpacity(this, &SHomeScreen::GetResourceAndSocialMediaButtonColor, OutButton));
+}
+
+TSharedRef<SWidget> SHomeScreen::CreateComboButtonMenuContentWidget()
+{
+	return SNew(SVerticalBox)
+
+	+ SVerticalBox::Slot()
+	.AutoHeight()
+	.Padding(8.f, 8.f)
+	[
+		SNew(SButton)
+		.ButtonStyle(FAppStyle::Get(), "HomeScreen.ComboButton.MenuButton")
+		.ContentPadding(FMargin(4.f, 4.f))
+		.OnClicked(this, &SHomeScreen::OnAutoloadLastProjectChanged, EAutoLoadProject::HomeScreen)
+		[
+			SNew(SHorizontalBox)
+
+			+ SHorizontalBox::Slot()
+			.AutoWidth()
+			[
+				SNew(SImage)
+				.Image(FAppStyle::GetBrush("Icons.Check"))
+				.Visibility(this, &SHomeScreen::GetAutoLoadProjectComboBoxContentIconVisibility, EAutoLoadProject::HomeScreen)
+			]
+
+			+ SHorizontalBox::Slot()
+			.AutoWidth()
+			[
+				SNew(STextBlock)
+				.Text(LOCTEXT("HomeScreenHomeTab", "Home Tab"))
+				.ColorAndOpacity(FStyleColors::White)
+			]
+		]
+	]
+
+	+ SVerticalBox::Slot()
+	.AutoHeight()
+	.Padding(8.f, 0.f)
+	[
+		SNew(SButton)
+		.ButtonStyle(FAppStyle::Get(), "HomeScreen.ComboButton.MenuButton")
+		.ContentPadding(FMargin(4.f, 4.f))
+		.OnClicked(this, &SHomeScreen::OnAutoloadLastProjectChanged, EAutoLoadProject::LastProject)
+		[
+			SNew(SHorizontalBox)
+
+			+ SHorizontalBox::Slot()
+			.AutoWidth()
+			[
+				SNew(SImage)
+				.Image(FAppStyle::GetBrush("Icons.Check"))
+				.Visibility(this, &SHomeScreen::GetAutoLoadProjectComboBoxContentIconVisibility, EAutoLoadProject::LastProject)
+
+			]
+
+			+ SHorizontalBox::Slot()
+			.AutoWidth()
+			[
+				SNew(STextBlock)
+				.Text(LOCTEXT("HomeScreenMostRecentProject", "Most Recent Project"))
+				.ColorAndOpacity(FStyleColors::White)
+			]
+		]
+	]
+
+	+ SVerticalBox::Slot()
+	.AutoHeight()
+	.Padding(8.f, 8.f)
+	[
+		SNew(SSeparator)
+		.Orientation(Orient_Horizontal)
+		.Thickness(1.f)
+		.SeparatorImage(FAppStyle::GetBrush("WhiteBrush"))
+		.ColorAndOpacity(FStyleColors::White25)
+	]
+
+	+ SVerticalBox::Slot()
+	.AutoHeight()
+	.Padding(8.f, 0.f, 8.f, 8.f)
+	[
+		SNew(SButton)
+		.ButtonStyle(FAppStyle::Get(), "HomeScreen.ComboButton.MenuButton")
+		.ContentPadding(FMargin(4.f, 4.f))
+		.OnClicked(this, &SHomeScreen::OnShowHomePageWhenNewAnnouncementOptionClicked)
+		[
+			SNew(SHorizontalBox)
+
+			+ SHorizontalBox::Slot()
+			.AutoWidth()
+			.VAlign(VAlign_Center)
+			[
+				SNew(SBox)
+				.HeightOverride(16.f)
+				[
+					SNew(SImage)
+					.Image(FAppStyle::GetBrush("Icons.Check"))
+					.Visibility_Lambda([this] ()
+					{
+						return bShowHomePageOption ? EVisibility::Visible : EVisibility::Hidden;
+					})
+				]
+			]
+
+			+ SHorizontalBox::Slot()
+			.AutoWidth()
+			[
+				SNew(STextBlock)
+				.WrapTextAt(222.f)
+				.Text(LOCTEXT("HomeScreenShowHomePage", "Show Home page when new announcements are available"))
+				.ColorAndOpacity_Lambda([this] ()
+				{
+					return bShowHomePageOption ? FStyleColors::White : FStyleColors::White25;
+				})
+			]
+		]
+	];
+}
+
+void SHomeScreen::CheckInternetConnection()
+{
+	// Fire another request only if the previous one finished
+	if (HttpRetryRequest.IsValid())
+	{
+		return;
+	}
+
+	bIsRequestFinished = false;
+
+	HttpRetryRequest = HttpRetryManager->CreateRequest(
+		FHttpRetrySystem::FRetryLimitCountSetting(3),
+		FHttpRetrySystem::FRetryTimeoutRelativeSecondsSetting(3),
+		FHttpRetrySystem::FRetryResponseCodes(),
+		FHttpRetrySystem::FRetryVerbs(),
+		FHttpRetrySystem::FRetryDomainsPtr(),
+		FHttpRetrySystem::FRetryLimitCountSetting(3),
+		FHttpRetrySystem::FExponentialBackoffCurve());
+
+	HttpRetryRequest->SetURL(TEXT("https://www.google.com/generate_204"));
+	HttpRetryRequest->SetVerb("GET");
+
+	HttpRetryRequest->OnProcessRequestComplete().BindLambda([this](FHttpRequestPtr Request, FHttpResponsePtr Response, bool bWasSuccessful)
+		{
+			bIsRequestFinished = true;
+			bIsConnected = bWasSuccessful;
+			HttpRetryRequest.Reset();
+		});
+
+	HttpRetryRequest->ProcessRequest();
+}
+
+bool SHomeScreen::IsConnectedToInternet() const
+{
+	return bIsConnected;
+}
+
+bool SHomeScreen::IsMainSectionEnabled(EMainSectionMenu InHomeSection) const
+{
+	return IsConnectedToInternet() || InHomeSection == EMainSectionMenu::Home;
+}
+
+int32 SHomeScreen::GetNoInternetIconIndex() const
+{
+	return bIsRequestFinished? 0 /** No Internet */ : 1 /** Retrying */;
+}
+
+void SHomeScreen::OnNavigateToSection(EMainSectionMenu InSectionToNavigate)
+{
+	OnMainHomeSectionChanged(ECheckBoxState::Checked, InSectionToNavigate);
+}
+
+FText SHomeScreen::GetAutoLoadProjectComboBoxLabelText() const
+{
+	return AutoLoadProjectComboBoxSelection == EAutoLoadProject::HomeScreen ? LOCTEXT("HomeScreenHomeTab", "Home Tab") : LOCTEXT("HomeScreenMostRecentProject", "Most Recent Project");
+}
+
+EVisibility SHomeScreen::GetAutoLoadProjectComboBoxContentIconVisibility(EAutoLoadProject InAutoLoadOption) const
+{
+	return AutoLoadProjectComboBoxSelection == InAutoLoadOption ? EVisibility::Visible : EVisibility::Hidden;
+}
+
+FReply SHomeScreen::OnInternetConnectionRetried()
+{
+	CheckInternetConnection();
+	return FReply::Handled();
+}
+
+#undef LOCTEXT_NAMESPACE
diff --git a/Engine/Source/Editor/MainFrame/Private/HomeScreen/SHomeScreen.h b/Engine/Source/Editor/MainFrame/Private/HomeScreen/SHomeScreen.h
new file mode 100644
--- /dev/null
+++ b/Engine/Source/Editor/MainFrame/Private/HomeScreen/SHomeScreen.h
@@ -0,0 +1,169 @@
+﻿// Copyright Epic Games, Inc. All Rights Reserved.
+
+#pragma once
+#include "Engine/TimerHandle.h"
+#include "HomeScreenCommon.h"
+#include "HomeScreenWeb.h"
+#include "HttpRetrySystem.h"
+#include "Input/Reply.h"
+#include "Styling/SlateColor.h"
+#include "Styling/SlateTypes.h"
+#include "Widgets/SCompoundWidget.h"
+
+class SButton;
+class SCheckBox;
+class SComboButton;
+class SWebBrowser;
+template <typename ItemType> class STileView;
+
+enum class EAutoLoadProject
+{
+	HomeScreen,
+	LastProject,
+	MAX
+};
+
+class SHomeScreen : public SCompoundWidget
+{
+public:
+	SLATE_BEGIN_ARGS(SHomeScreen)
+	{}
+
+	SLATE_END_ARGS()
+
+	SHomeScreen();
+	virtual ~SHomeScreen() override;
+
+	/** Constructs the HomeScreen widget */
+	void Construct(const FArguments& InArgs, const TSharedPtr<SWebBrowser> InWebBrowser, const TWeakObjectPtr<UHomeScreenWeb> InWebObject);
+
+private:
+	/** Opens the Create Project dialog */
+	FReply OnCreateProjectDialog(bool bInAllowProjectOpening, bool bInAllowProjectCreation);
+
+	/** Handles social media button clicks */
+	FReply OnSocialMediaClicked(FString InURL);
+
+	/** Checks if a main home section is selected */
+	ECheckBoxState IsMainHomeSectionChecked(EMainSectionMenu InMainHomeSelection) const;
+
+	/** Handles changes to the main home section */
+	void OnMainHomeSectionChanged(ECheckBoxState InCheckBoxState, EMainSectionMenu InMainHomeSelection);
+
+	/** Loads settings from config */
+	void LoadSettings();
+
+	/** Saves settings to config */
+	void SaveSettings();
+
+	/** Handles changes to autoload last project selection */
+	FReply OnAutoloadLastProjectChanged(EAutoLoadProject InAutoLoadOption);
+
+	/** Handles clicks for showing the home page when new announcements are available */
+	FReply OnShowHomePageWhenNewAnnouncementOptionClicked();
+
+	/** Checks if a checkbox is checked or hovered */
+	bool IsCheckBoxCheckedOrHovered(const TSharedPtr<SCheckBox> InCheckBox) const;
+
+	/** Gets the color of the main section checkbox icon and text */
+	FSlateColor GetMainSectionCheckBoxColor(const TSharedPtr<SCheckBox> InCheckBox) const;
+
+	/** Gets the color of resource and social media section icon and text */
+	FSlateColor GetResourceAndSocialMediaButtonColor(const TSharedPtr<SButton> InButton) const;
+
+	/** Creates a main section checkboxes */
+	void CreateMainSectionCheckBox(TSharedPtr<SCheckBox>& OutCheckBox, EMainSectionMenu InMainHomeSelection, const FText& InText, const FSlateBrush* InImage);
+
+	/** Creates resource buttons */
+	void CreateResourceButtons(TSharedPtr<SButton>& OutButton, FString InLink, const FText& InText, const FSlateBrush* InImage);
+
+	/** Creates social media buttons */
+	void CreateSocialMediaButtons(TSharedPtr<SButton>& OutButton, FString InLink, const FSlateBrush* InImage);
+
+	/** Creates the content widget for the combo button menu */
+	TSharedRef<SWidget> CreateComboButtonMenuContentWidget();
+
+	/** Checks the internet connection */
+	void CheckInternetConnection();
+
+	/** Returns whether the system is connected to the internet */
+	bool IsConnectedToInternet() const;
+
+	/** Returns whether the main section is enabled */
+	bool IsMainSectionEnabled(EMainSectionMenu InHomeSection) const;
+
+	/** Gets the index of the SWidgetSwitcher for the no internet icon, either NoInternet or Loading when retrying */
+	int32 GetNoInternetIconIndex() const;
+
+	/** Handles navigation to a section from a web request */
+	void OnNavigateToSection(EMainSectionMenu InSectionToNavigate);
+
+	/** Gets the label text for the autoload project combo box */
+	FText GetAutoLoadProjectComboBoxLabelText() const;
+
+	/** Get the visibility for the check icon beside the ComboBox options */
+	EVisibility GetAutoLoadProjectComboBoxContentIconVisibility(EAutoLoadProject InAutoLoadOption) const;
+
+	/** Executed when clicking on the Reconnect button when no internet connection is detected */
+	FReply OnInternetConnectionRetried();
+
+private:
+	/** Currently selected main home section */
+	EMainSectionMenu MainHomeSelection = EMainSectionMenu::Home;
+
+	/** Main section checkboxes */
+	TSharedPtr<SCheckBox> HomeCheckBox;
+	TSharedPtr<SCheckBox> NewsCheckBox;
+	TSharedPtr<SCheckBox> GettingStartedCheckBox;
+	TSharedPtr<SCheckBox> SampleProjectsCheckBox;
+
+	/** Resource section buttons */
+	TSharedPtr<SButton> ForumsButton;
+	TSharedPtr<SButton> DocumentationButton;
+	TSharedPtr<SButton> TutorialsButton;
+	TSharedPtr<SButton> RoadmapButton;
+	TSharedPtr<SButton> ReleaseNotesButton;
+
+	/** Http retry system used to check Internet connection */
+	TSharedPtr<FHttpRetrySystem::FManager> HttpRetryManager;
+	TSharedPtr<FHttpRetrySystem::FRequest> HttpRetryRequest;
+
+	/** Combo box for autoload project selection */
+	TSharedPtr<SComboButton> AutoLoadProjectComboBox;
+
+	/** Widget for project browser */
+	TSharedPtr<SWidget> ProjectBrowser;
+
+	/** Timer handle for checking internet connection */
+	FTimerHandle CheckInternetConnectionTimerHandle;
+
+	/** Path of the currently selected project */
+	FString CurrentSelectedProjectPath;
+
+	/** Whether to show the home page option */
+	bool bShowHomePageOption = false;
+
+	/** Whether the system is connected to the internet */
+	bool bIsConnected = true;
+
+	/** Whether the last request has finished */
+	bool bIsRequestFinished = true;
+
+	/** How many times we executed the function without being connected */
+	int32 TimerManagerCountOnceDisconnected = 0;
+
+	/** Max TimerManager function execution before retrying once disconnected */
+	int32 MaxTimerManagerCountOnceDisconnectedBeforeRetry = 5;
+
+	/** Whether to force a retry on the connection */
+	bool bForceRetry = false;
+
+	/** Selection for autoload project combo box */
+	EAutoLoadProject AutoLoadProjectComboBoxSelection;
+
+	/** WebBrowser widget passed down from the MainFrameModule */
+	TSharedPtr<SWebBrowser> WebBrowser;
+	
+	/** WebBrowser object passed down from the MainFrameModule */
+	TWeakObjectPtr<UHomeScreenWeb> WebObject;
+};
diff --git a/Engine/Source/Editor/MainFrame/Private/MainFrameModule.cpp b/Engine/Source/Editor/MainFrame/Private/MainFrameModule.cpp
--- a/Engine/Source/Editor/MainFrame/Private/MainFrameModule.cpp
+++ b/Engine/Source/Editor/MainFrame/Private/MainFrameModule.cpp
@@ -36,9 +36,66 @@
 #include "ProfilingDebugging/StallDetector.h"
 #include "Interfaces/IEditorMainFrameProvider.h"
 
+// HomeScreen Includes
+#include "HomeScreen/HomeScreenWeb.h"
+#include "HomeScreen/SHomeScreen.h" 
+#include "IWebBrowserSingleton.h"
+#include "IWebBrowserWindow.h"
+#include "SWebBrowser.h"
+#include "ToolMenus.h"
+#include "WebBrowserModule.h"
+#include "Widgets/Docking/SDockTab.h"
+
 DEFINE_LOG_CATEGORY(LogMainFrame);
 #define LOCTEXT_NAMESPACE "FMainFrameModule"
 
+#if WITH_EDITOR
+TStrongObjectPtr<UHomeScreenWeb> FMainFrameModule::WebObject = nullptr;
+TSharedPtr<IWebBrowserWindow> FMainFrameModule::WebBrowserWindow = nullptr;
+TSharedPtr<SWebBrowser> FMainFrameModule::WebBrowser = nullptr;
+#endif
+
+namespace UE::Editor::HomeScreen
+{
+	static FAutoConsoleVariable CVarEnableHomeScreen(
+	TEXT("HomeScreen.EnableHomeScreen"),
+	false,
+	TEXT("Whether or not to enable the new HomeScreen.")
+);
+
+	void SetEnableHomeScreenFromCmdLine()
+	{
+		const auto OriginalCmdLine = FCommandLine::GetOriginal();
+		FString OutValue;
+		if (FParse::Value(OriginalCmdLine, TEXT("-HomeScreen.EnableHomeScreen="), OutValue); OutValue == TEXT("1"))
+		{
+			CVarEnableHomeScreen->Set(true);
+		}
+	}
+
+	static FDelayedAutoRegisterHelper EnableHomeScreenCVarRegistration(
+		EDelayedRegisterRunPhase::FileSystemReady,
+		[]()
+		{
+			SetEnableHomeScreenFromCmdLine();
+		});
+
+	static bool IsHomeScreenEnabled()
+	{
+		static bool bIsHomeScreenEnabled = false;
+		UE_CALL_ONCE([&]()
+		{
+			if (const IConsoleVariable* CVar = IConsoleManager::Get().FindConsoleVariable(TEXT("HomeScreen.EnableHomeScreen")))
+			{
+				ensureAlwaysMsgf(!EnumHasAnyFlags(CVar->GetFlags(), ECVF_Default), TEXT("The CVar should have already been set from commandline."));
+				bIsHomeScreenEnabled = CVar->GetBool();
+			}
+		});
+
+		return bIsHomeScreenEnabled;
+	}
+}
+
 static FAutoConsoleCommand ResizeMainFrameCommand(
 	TEXT("Editor.ResizeMainFrame"),
 	TEXT(""),
@@ -616,12 +673,40 @@ void FMainFrameModule::StartupModule( )
 	HotReloadModule.OnModuleCompilerFinished().AddRaw( this, &FMainFrameModule::HandleLevelEditorModuleCompileFinished );
 	FCoreUObjectDelegates::ReloadCompleteDelegate.AddRaw( this, &FMainFrameModule::HandleReloadFinished );
 
-#if WITH_EDITOR
 	ISourceCodeAccessModule& SourceCodeAccessModule = FModuleManager::LoadModuleChecked<ISourceCodeAccessModule>("SourceCodeAccess");
 	SourceCodeAccessModule.OnLaunchingCodeAccessor().AddRaw( this, &FMainFrameModule::HandleCodeAccessorLaunching );
 	SourceCodeAccessModule.OnDoneLaunchingCodeAccessor().AddRaw( this, &FMainFrameModule::HandleCodeAccessorLaunched );
 	SourceCodeAccessModule.OnOpenFileFailed().AddRaw( this, &FMainFrameModule::HandleCodeAccessorOpenFileFailed );
-#endif
+
+	// HomeScreen Registration
+		if (UE::Editor::HomeScreen::IsHomeScreenEnabled())
+    	{
+    		FGlobalTabmanager::Get()->RegisterTabSpawner(
+    			"HomeScreen", FOnSpawnTab::CreateRaw(this, &FMainFrameModule::SpawnHomeScreenWindow))
+    			.SetAutoGenerateMenuEntry(false)
+    			.SetMenuType(ETabSpawnerMenuType::Hidden)
+    			.SetDisplayName(FText::GetEmpty());
+    
+    		// Register menu extension
+    		UToolMenus::RegisterStartupCallback(
+    			FSimpleMulticastDelegate::FDelegate::CreateLambda([] ()
+    				{
+    					UToolMenu* Menu = UToolMenus::Get()->ExtendMenu("LevelEditor.MainMenu.Tools");
+    
+    					FToolMenuSection& Section = Menu->AddSection("HomeScreenSection", FText::FromString("Home Screen"));
+    					Section.AddMenuEntry(
+    						"HomeScreen",
+    						FText::FromString("Home Screen"),
+    						FText::FromString("Opens the Home Screen tab"),
+    						FSlateIcon(FAppStyle::GetAppStyleSetName(), "HomeScreen.Home"),
+    						FUIAction(FExecuteAction::CreateLambda([] ()
+    							{
+    								FGlobalTabmanager::Get()->TryInvokeTab(FTabId("HomeScreen"));
+    							}))
+    					);
+    				})
+    		);
+    	}
 
 	ModuleCompileStartTime = 0.0f;
 
@@ -653,7 +738,6 @@ void FMainFrameModule::ShutdownModule( )
 		HotReloadModule.OnModuleCompilerFinished().RemoveAll( this );
 	}
 
-#if WITH_EDITOR
 	if(FModuleManager::Get().IsModuleLoaded("SourceCodeAccess"))
 	{
 		ISourceCodeAccessModule& SourceCodeAccessModule = FModuleManager::GetModuleChecked<ISourceCodeAccessModule>("SourceCodeAccess");
@@ -661,8 +745,11 @@ void FMainFrameModule::ShutdownModule( )
 		SourceCodeAccessModule.OnDoneLaunchingCodeAccessor().RemoveAll( this );
 		SourceCodeAccessModule.OnOpenFileFailed().RemoveAll( this );
 	}
-#endif
 
+	// Clean up static object/widget used for the HomeScreen Tab
+	WebObject.Reset();
+	WebBrowserWindow.Reset();
+	WebBrowser.Reset();
 	ProjectDialogProvider.UnRegister();
 }
 
@@ -864,6 +951,54 @@ void FMainFrameModule::HandleCodeAccessorLaunching()
 	CodeAccessorNotificationPtr.Pin()->SetCompletionState(SNotificationItem::CS_Pending);
 }
 
+TSharedRef<SDockTab> FMainFrameModule::SpawnHomeScreenWindow(const FSpawnTabArgs& InArgs)
+{
+	if (!WebObject.IsValid())
+	{
+		WebObject = TStrongObjectPtr(NewObject<UHomeScreenWeb>());
+	}
+
+	if (!WebBrowserWindow.IsValid())
+	{
+		IWebBrowserModule& WebBrowserModule = IWebBrowserModule::Get();
+		if (!IWebBrowserModule::IsAvailable() || !WebBrowserModule.IsWebModuleAvailable())
+		{
+			FMessageDialog::Open(EAppMsgType::Ok, LOCTEXT("WebBrowserModuleNotAvailable", "Failed to load the plugin. Please enable Web WebBrowserWindow in the plugin manager to use the Web Panels."));
+		}
+		IWebBrowserSingleton* WebBrowserSingleton = WebBrowserModule.GetSingleton();
+
+		FCreateBrowserWindowSettings WindowSettings;
+		WindowSettings.InitialURL = TEXT("https://editor-website.ol.epicgames.net/en-US");
+		WindowSettings.bUseTransparency = true;
+		WindowSettings.BrowserFrameRate = 60;
+
+		WebBrowserWindow = WebBrowserSingleton->CreateBrowserWindow(WindowSettings);
+	}
+
+	if (!WebBrowser.IsValid())
+	{
+		SAssignNew(WebBrowser, SWebBrowser, WebBrowserWindow)
+		.ShowControls(false)
+		.SupportsTransparency(true)
+		.ShowControls(false)
+		.ShowAddressBar(false);
+
+		constexpr bool bIsPermanent = true;
+		WebBrowser->BindUObject(TEXT("uebridge"), WebObject.Get(), bIsPermanent);
+		WebBrowserWindow->Reload();
+	}
+
+	TSharedRef<SDockTab> HomeScreenTab = SNew(SDockTab)
+	.TabRole(ETabRole::MajorTab)
+	.IconColor(FStyleColors::White.GetSpecifiedColor())
+	.ContentPadding(FMargin(0));
+
+	HomeScreenTab->SetTabIcon(FAppStyle::Get().GetBrush("HomeScreen.Home"));
+
+	HomeScreenTab->SetContent(SNew(SHomeScreen, WebBrowser, TWeakObjectPtr(WebObject.Get())));
+
+	return HomeScreenTab;
+}
 
 void FMainFrameModule::HandleCodeAccessorOpenFileFailed(const FString& Filename)
 {
diff --git a/Engine/Source/Editor/MainFrame/Private/MainFrameModule.h b/Engine/Source/Editor/MainFrame/Private/MainFrameModule.h
--- a/Engine/Source/Editor/MainFrame/Private/MainFrameModule.h
+++ b/Engine/Source/Editor/MainFrame/Private/MainFrameModule.h
@@ -13,6 +13,10 @@
 #include "Misc/CompilationResult.h"
 #include "Interfaces/IEditorMainFrameProvider.h"
 
+class SWebBrowser;
+class IWebBrowserWindow;
+class UHomeScreenWeb;
+
 /** 
  * Utility class, which hooks into the editor's main window startup and 
  * supplies a project dialog window if the project hasn't been set. 
@@ -243,6 +247,9 @@ private:
 	// Handles launching code accessor
 	void HandleCodeAccessorLaunching( );
 
+	// Spawns the home screen tab
+	TSharedRef<SDockTab> SpawnHomeScreenWindow(const FSpawnTabArgs& InArgs);
+
 	// Reset delegate
 	void ClearDelayedShowMainFrameDelegate()
 	{
@@ -306,4 +313,13 @@ private:
 
 	// Set to true if we are in the middle of recreating the default layout because a previous load attempt failed
 	bool bIsReCreatingDefaultLayoutAsFallback = false;
+
+	// Bridge object for browser navigation
+	static TStrongObjectPtr<UHomeScreenWeb> WebObject;
+
+	// Web browser window instance
+	static TSharedPtr<IWebBrowserWindow> WebBrowserWindow;
+
+	// Web browser widget instance
+	static TSharedPtr<SWebBrowser> WebBrowser;
 };
