[HomeScreen] Fix possible issue with TabManager Registering and WebBrowser usage if not Available

diff --git a/Engine/Source/Editor/MainFrame/Private/MainFrameModule.cpp b/Engine/Source/Editor/MainFrame/Private/MainFrameModule.cpp
--- a/Engine/Source/Editor/MainFrame/Private/MainFrameModule.cpp
+++ b/Engine/Source/Editor/MainFrame/Private/MainFrameModule.cpp
@@ -750,6 +750,10 @@ void FMainFrameModule::ShutdownModule( )
 	WebObject.Reset();
 	WebBrowserWindow.Reset();
 	WebBrowser.Reset();
+
+	// Unregister HomeScreen tab on shutdown
+	FGlobalTabmanager::Get()->UnregisterTabSpawner(TEXT("HomeScreen"));
+
 	ProjectDialogProvider.UnRegister();
 }
 
@@ -953,6 +957,11 @@ void FMainFrameModule::HandleCodeAccessorLaunching()
 
 TSharedRef<SDockTab> FMainFrameModule::SpawnHomeScreenWindow(const FSpawnTabArgs& InArgs)
 {
+	TSharedRef<SDockTab> HomeScreenTab = SNew(SDockTab)
+		.TabRole(ETabRole::MajorTab)
+		.IconColor(FStyleColors::White.GetSpecifiedColor())
+		.ContentPadding(FMargin(0));
+
 	if (!WebObject.IsValid())
 	{
 		WebObject = TStrongObjectPtr(NewObject<UHomeScreenWeb>());
@@ -964,6 +973,7 @@ TSharedRef<SDockTab> FMainFrameModule::SpawnHomeScreenWindow(const FSpawnTabArgs
 		if (!IWebBrowserModule::IsAvailable() || !WebBrowserModule.IsWebModuleAvailable())
 		{
 			FMessageDialog::Open(EAppMsgType::Ok, LOCTEXT("WebBrowserModuleNotAvailable", "Failed to load the plugin. Please enable Web WebBrowserWindow in the plugin manager to use the Web Panels."));
+			return HomeScreenTab;
 		}
 		IWebBrowserSingleton* WebBrowserSingleton = WebBrowserModule.GetSingleton();
 
@@ -988,11 +998,6 @@ TSharedRef<SDockTab> FMainFrameModule::SpawnHomeScreenWindow(const FSpawnTabArgs
 		WebBrowserWindow->Reload();
 	}
 
-	TSharedRef<SDockTab> HomeScreenTab = SNew(SDockTab)
-	.TabRole(ETabRole::MajorTab)
-	.IconColor(FStyleColors::White.GetSpecifiedColor())
-	.ContentPadding(FMargin(0));
-
 	HomeScreenTab->SetTabIcon(FAppStyle::Get().GetBrush("HomeScreen.Home"));
 
 	HomeScreenTab->SetContent(SNew(SHomeScreen, WebBrowser, TWeakObjectPtr(WebObject.Get())));
