Build Engine Localization - Text conflict from LOCTEXT macro for namespace 'SHomeScreen' and key 'ProjectHomeDocumentation'

diff --git a/Engine/Source/Editor/MainFrame/Private/HomeScreen/SHomeScreen.cpp b/Engine/Source/Editor/MainFrame/Private/HomeScreen/SHomeScreen.cpp
--- a/Engine/Source/Editor/MainFrame/Private/HomeScreen/SHomeScreen.cpp
+++ b/Engine/Source/Editor/MainFrame/Private/HomeScreen/SHomeScreen.cpp
@@ -264,7 +264,7 @@ void SHomeScreen::Construct(const FArguments& InArgs, const TSharedPtr<SWebBrows
 					.Padding(8.f, 0.f, 8.f, 4.f)
 					[
 						SNew(STextBlock)
-						.Text(LOCTEXT("ProjectHomeDocumentation", "Resources & Community"))
+						.Text(LOCTEXT("ResourceAndCommunity", "Resources & Community"))
 						.Font(FAppStyle::GetFontStyle("BoldFont"))
 						.ColorAndOpacity(FStyleColors::White)
 					]
@@ -948,7 +948,7 @@ void SHomeScreen::OnNavigateToSection(EMainSectionMenu InSectionToNavigate)
 
 FText SHomeScreen::GetAutoLoadProjectComboBoxLabelText() const
 {
-	return AutoLoadProjectComboBoxSelection == EAutoLoadProject::HomeScreen ? LOCTEXT("HomeScreenHomeTab", "Home Tab") : LOCTEXT("HomeScreenMostRecentProject", "Most Recent Project");
+	return AutoLoadProjectComboBoxSelection == EAutoLoadProject::HomeScreen ? LOCTEXT("HomeScreenComboHomeTab", "Home Tab") : LOCTEXT("HomeScreenComboMostRecentProject", "Most Recent Project");
 }
 
 EVisibility SHomeScreen::GetAutoLoadProjectComboBoxContentIconVisibility(EAutoLoadProject InAutoLoadOption) const
