[HomeScreen] Further URL security issue

diff --git a/Engine/Source/Editor/MainFrame/Private/HomeScreen/HomeScreenWeb.cpp b/Engine/Source/Editor/MainFrame/Private/HomeScreen/HomeScreenWeb.cpp
--- a/Engine/Source/Editor/MainFrame/Private/HomeScreen/HomeScreenWeb.cpp
+++ b/Engine/Source/Editor/MainFrame/Private/HomeScreen/HomeScreenWeb.cpp
@@ -1,6 +1,10 @@
 ﻿// Copyright Epic Games, Inc. All Rights Reserved.
 
 #include "HomeScreenWeb.h"
+#include "Internationalization/Regex.h"
+#include "Misc/MessageDialog.h"
+
+#define LOCTEXT_NAMESPACE "HomeScreenWeb"
 
 void UHomeScreenWeb::NavigateTo(EMainSectionMenu InSectionToNavigate)
 {
@@ -11,5 +15,20 @@ void UHomeScreenWeb::NavigateTo(EMainSectionMenu InSectionToNavigate)
 
 void UHomeScreenWeb::OpenWebPage(const FString& InURL) const
 {
-	FPlatformProcess::LaunchURL(*InURL, nullptr, nullptr);
+	const FText ErrorText = FText::Format(LOCTEXT("LaunchingURLNotAllowed", "Given URL is not registered as an allowed one:\n {0}"), FText::FromString(InURL));
+
+	const FString EpicOriginRegex = TEXT("^https?://(www|www2|localhost|dev|docs|forums|company|legal)?(\\.staging)?\\.?(unrealengine|epicgames|twinmotion|metahuman|realityscan|kidswebservices|fab|quixel)(\\.com|-(ci|gamedev)\\.ol\\.epicgames\\.net)(/[a-zA-Z0-9._~:/?#@!$&'()*+,;=%-]*/?)*$");
+	const FRegexPattern URLPattern(EpicOriginRegex);
+	FRegexMatcher Matcher(URLPattern, InURL);
+
+	if (Matcher.FindNext())
+	{
+		FPlatformProcess::LaunchURL(*InURL, nullptr, nullptr);
+	}
+	else
+	{
+		FMessageDialog::Open(EAppMsgType::Ok, ErrorText);
+	}
 }
+
+#undef LOCTEXT_NAMESPACE
