[HomeScreen] CaseSensitive and Port issue with URL regex

diff --git a/Engine/Source/Editor/MainFrame/Private/HomeScreen/HomeScreenWeb.cpp b/Engine/Source/Editor/MainFrame/Private/HomeScreen/HomeScreenWeb.cpp
--- a/Engine/Source/Editor/MainFrame/Private/HomeScreen/HomeScreenWeb.cpp
+++ b/Engine/Source/Editor/MainFrame/Private/HomeScreen/HomeScreenWeb.cpp
@@ -17,8 +17,8 @@ void UHomeScreenWeb::OpenWebPage(const FString& InURL) const
 {
 	const FText ErrorText = FText::Format(LOCTEXT("LaunchingURLNotAllowed", "Given URL is not registered as an allowed one:\n {0}"), FText::FromString(InURL));
 
-	const FString EpicOriginRegex = TEXT("^https?://(www|www2|localhost|dev|docs|forums|company|legal)?(\\.staging)?\\.?(unrealengine|epicgames|twinmotion|metahuman|realityscan|kidswebservices|fab|quixel)(\\.com|-(ci|gamedev)\\.ol\\.epicgames\\.net)(/[a-zA-Z0-9._~:/?#@!$&'()*+,;=%-]*/?)*$");
-	const FRegexPattern URLPattern(EpicOriginRegex);
+	const FString EpicOriginRegex = TEXT("^https?://(www|www2|localhost|dev|docs|forums|company|legal)?(\\.staging)?\\.?(unrealengine|epicgames|twinmotion|metahuman|realityscan|kidswebservices|fab|quixel)(\\.com|-(ci|gamedev)\\.ol\\.epicgames\\.net)(:\\d{1,5})?(/[a-zA-Z0-9._~:/?#@!$&'()*+,;=%-]*/?)*$");
+	const FRegexPattern URLPattern(EpicOriginRegex, ERegexPatternFlags::CaseInsensitive);
 	FRegexMatcher Matcher(URLPattern, InURL);
 
 	if (Matcher.FindNext())
