[HomeScreen] HomeScreen style change based on feedbacks

diff --git a/Engine/Source/Editor/EditorStyle/Private/StarshipStyle.cpp b/Engine/Source/Editor/EditorStyle/Private/StarshipStyle.cpp
--- a/Engine/Source/Editor/EditorStyle/Private/StarshipStyle.cpp
+++ b/Engine/Source/Editor/EditorStyle/Private/StarshipStyle.cpp
@@ -8534,8 +8534,8 @@ void FStarshipEditorStyle::FStyle::SetupHomeScreenStyles()
 	Set("HomeScreen.Tutorial", new IMAGE_BRUSH_SVG("Starship/HomeScreen/Tutorial", Icon20x20));
 	Set("HomeScreen.Roadmap", new IMAGE_BRUSH_SVG("Starship/HomeScreen/Roadmap", Icon20x20));
 	Set("HomeScreen.ReleaseNotes", new IMAGE_BRUSH_SVG("Starship/HomeScreen/ReleaseNotes", Icon20x20));
-	Set("HomeScreen.PlusCircle", new IMAGE_BRUSH_SVG("Starship/HomeScreen/PlusCircleEmpty", Icon20x20));
-	Set("HomeScreen.FolderOpen", new IMAGE_BRUSH_SVG("Starship/HomeScreen/FolderOpen", Icon20x20));
+	Set("HomeScreen.PlusCircle", new IMAGE_BRUSH_SVG("Starship/HomeScreen/PlusCircleEmpty", Icon16x16));
+	Set("HomeScreen.FolderOpen", new IMAGE_BRUSH_SVG("Starship/HomeScreen/FolderOpen", Icon16x16));
 
 	Set("SocialMedia.Facebook", new IMAGE_BRUSH_SVG("Starship/HomeScreen/Facebook", Icon20x20));
 	Set("SocialMedia.Twitter", new IMAGE_BRUSH_SVG("Starship/HomeScreen/Twitter", Icon20x20));
@@ -8543,13 +8543,23 @@ void FStarshipEditorStyle::FStyle::SetupHomeScreenStyles()
 	Set("SocialMedia.Twitch", new IMAGE_BRUSH_SVG("Starship/HomeScreen/Twitch", Icon20x20));
 	Set("SocialMedia.Instagram", new IMAGE_BRUSH_SVG("Starship/HomeScreen/Instagram", Icon20x20));
 
-	Set("HomeScreen.Background", new IMAGE_BRUSH("Starship/HomeScreen/Background", FVector2D(4096.f, 2304.f)));
-	Set("HomeScreen.FakeWeb", new IMAGE_BRUSH("Starship/HomeScreen/FakeWeb", FVector2D(2560.f, 1440.f)));
 	Set("HomeScreen.RecentProjects.Background", new FSlateRoundedBoxBrush(COLOR("#FFFFFF0D"), 16.f, FStyleColors::White25, 1.f));
 
 	Set("HomeScreen.NoInternet.Border", new FSlateRoundedBoxBrush(COLOR("#FFFFFF0D"), 8.f, COLOR("#FFFFFF59"), 1.f));
 	Set("HomeScreen.NoInternet.Icon", new IMAGE_BRUSH_SVG("Starship/HomeScreen/NoInternet", FVector2D(24.f, 24.f)));
 
+	FSlateFontInfo NewAndCreateProjectButtonFont = DEFAULT_FONT("Regular", 11);
+    Set("HomeScreen.Font.NewAndCreateButton", NewAndCreateProjectButtonFont);
+
+	FSlateFontInfo DefaultFontSize = DEFAULT_FONT("Regular", 12);
+	Set("HomeScreen.Font.DefaultSize", DefaultFontSize);
+
+	FSlateFontInfo DefaultBoldFontSize = DEFAULT_FONT("Bold", 12);
+	Set("HomeScreen.Font.DefaultBoldSize", DefaultBoldFontSize);
+
+	FSlateFontInfo RecentProjectFont = DEFAULT_FONT("Bold", 20);
+	Set("HomeScreen.Font.RecentProject", RecentProjectFont);
+
 	FComboButtonStyle LoadOnStartupStyle = FAppStyle::Get().GetWidgetStyle<FComboButtonStyle>("ComboButton");
 	LoadOnStartupStyle.ButtonStyle = FButtonStyle()
 		.SetNormal(FSlateRoundedBoxBrush(COLOR("#FFFFFF0D"), 6.f, COLOR("#FFFFFF59"), 1.f))
@@ -8559,15 +8569,15 @@ void FStarshipEditorStyle::FStyle::SetupHomeScreenStyles()
 		.SetNormalPadding(FMargin(12.0f, 4.0f))
 		.SetPressedPadding(FMargin(12.0f, 4.0f));
 
-
-	LoadOnStartupStyle.SetMenuBorderBrush(FSlateRoundedBoxBrush(FStyleColors::Panel, 6.f, FStyleColors::White25, 2.f));
+	LoadOnStartupStyle.SetMenuBorderBrush(FSlateRoundedBoxBrush(FStyleColors::Panel, 0.f, FStyleColors::White25, 1.f));
 	LoadOnStartupStyle.SetMenuBorderPadding(0.f);
 
+	Set("HomeScreen.ComboButton.HoveredEntry", new FSlateRoundedBoxBrush(FStyleColors::Input, 0.f));
+	Set("HomeScreen.ComboButton", LoadOnStartupStyle);
+
 	FButtonStyle LoadOnStartupMenuButtonStyle = FAppStyle::Get().GetWidgetStyle<FButtonStyle>("NoBorder");
 	LoadOnStartupMenuButtonStyle.SetHovered(FSlateRoundedBoxBrush(FStyleColors::Hover, 4.0f));
-	Set("HomeScreen.ComboButton", LoadOnStartupStyle);
 	Set("HomeScreen.ComboButton.MenuButton", LoadOnStartupMenuButtonStyle);
-	Set("HomeScreen.ComboButton.MenuBorder", new FSlateRoundedBoxBrush(FStyleColors::Panel, 8.f, FStyleColors::White25, 2.f));
 
 	FButtonStyle CreateNewProjectButton = FButtonStyle()
 		.SetNormal(FSlateRoundedBoxBrush(FStyleColors::AccentBlue, 4.f))
@@ -8591,9 +8601,6 @@ void FStarshipEditorStyle::FStyle::SetupHomeScreenStyles()
 		.SetPressedPadding(FMargin(0.f, 4.0f));
 	Set("HomeScreen.MyFolderButton", MyFolderButton);
 
-	FSlateFontInfo RecentProjectFont = DEFAULT_FONT("Bold", 18);
-	Set("HomeScreen.RecentProjectFont", RecentProjectFont);
-
 	FButtonStyle SeeAllProjectsButton = FButtonStyle()
 		.SetNormal(FSlateRoundedBoxBrush(FStyleColors::Transparent, 4.0f, COLOR("FFFFFF59"), 1.f))
 		.SetHovered(FSlateRoundedBoxBrush(COLOR("FFFFFF0D"), 4.0f, COLOR("FFFFFF59"), 1.f))
@@ -8618,7 +8625,7 @@ void FStarshipEditorStyle::FStyle::SetupHomeScreenStyles()
 
 	FCheckBoxStyle MainMenuSectionCheckBox = FAppStyle::GetWidgetStyle<FCheckBoxStyle>("ToggleButtonCheckbox");
 	MainMenuSectionCheckBox
-		.SetUncheckedHoveredImage(FSlateRoundedBoxBrush(FStyleColors::DropdownOutline, 4.0f))
+		.SetUncheckedHoveredImage(FSlateNoResource())
 		.SetUncheckedPressedImage(FSlateRoundedBoxBrush(FStyleColors::DropdownOutline, 4.0f))
 		.SetCheckedImage(FSlateRoundedBoxBrush(FStyleColors::DropdownOutline, 4.0f))
 		.SetCheckedHoveredImage(FSlateRoundedBoxBrush(FStyleColors::Hover, 4.0f));
diff --git a/Engine/Source/Editor/GameProjectGeneration/Private/GameProjectGenerationModule.cpp b/Engine/Source/Editor/GameProjectGeneration/Private/GameProjectGenerationModule.cpp
--- a/Engine/Source/Editor/GameProjectGeneration/Private/GameProjectGenerationModule.cpp
+++ b/Engine/Source/Editor/GameProjectGeneration/Private/GameProjectGenerationModule.cpp
@@ -160,6 +160,16 @@ void FGameProjectGenerationModule::ClearSupportedTargetPlatforms()
 	GameProjectUtils::ClearSupportedTargetPlatforms();
 }
 
+bool FGameProjectGenerationModule::CreateProject(const FProjectInformation& InProjectInfo, FText& OutFailReason, FText& OutFailLog, TArray<FString>* OutCreatedFiles)
+{
+	return GameProjectUtils::CreateProject(InProjectInfo, OutFailReason, OutFailLog, OutCreatedFiles);
+}
+
+bool FGameProjectGenerationModule::OpenProject(const FString& InProjectFile, FText& OutFailReason)
+{
+	return GameProjectUtils::OpenProject(InProjectFile, OutFailReason);
+}
+
 void FGameProjectGenerationModule::LoadTemplateCategories()
 {
 	// Now discover and all data driven templates
diff --git a/Engine/Source/Editor/GameProjectGeneration/Public/GameProjectGenerationModule.h b/Engine/Source/Editor/GameProjectGeneration/Public/GameProjectGenerationModule.h
--- a/Engine/Source/Editor/GameProjectGeneration/Public/GameProjectGenerationModule.h
+++ b/Engine/Source/Editor/GameProjectGeneration/Public/GameProjectGenerationModule.h
@@ -19,6 +19,7 @@
 class FText;
 class UClass;
 struct FModuleContextInfo;
+struct FProjectInformation;
 struct FSlateBrush;
 struct FTemplateCategory;
 
@@ -130,6 +131,12 @@ public:
 
 public:
 
+	/** Creates the specified project file and all required folders. If TemplateFile is non-empty, it will be used as the template for creation. On failure, OutFailReason will be populated. */
+	virtual bool CreateProject(const FProjectInformation& InProjectInfo, FText& OutFailReason, FText& OutFailLog, TArray<FString>* OutCreatedFiles = nullptr);
+
+	/** Opens the specified project, if it exists. Returns true if the project file is valid. On failure, OutFailReason will be populated. */
+	virtual bool OpenProject(const FString& InProjectFile, FText& OutFailReason);
+
 	// Non DLL-exposed access to template categories
 	TSharedPtr<const FTemplateCategory> GetCategory(FName Type) const { return TemplateCategories.FindRef(Type); }
 
diff --git a/Engine/Source/Editor/MainFrame/Private/Frame/MainFrameActions.cpp b/Engine/Source/Editor/MainFrame/Private/Frame/MainFrameActions.cpp
--- a/Engine/Source/Editor/MainFrame/Private/Frame/MainFrameActions.cpp
+++ b/Engine/Source/Editor/MainFrame/Private/Frame/MainFrameActions.cpp
@@ -406,26 +406,7 @@ void FMainFrameActionCallbacks::NewProject( bool bAllowProjectOpening, bool bAll
 		return;
 	}
 
-	FText Title;
-	if (bAllowProjectOpening && bAllowProjectCreate)
-	{
-		Title = LOCTEXT( "SelectProjectWindowHeader", "Select Project");
-	}
-	else if (bAllowProjectOpening)
-	{
-		Title = LOCTEXT( "OpenProjectWindowHeader", "Open Project");
-	}
-	else
-	{
-		Title = LOCTEXT( "NewProjectWindowHeader", "New Project");
-	}
-
-	TSharedRef<SWindow> NewProjectWindow =
-		SNew(SWindow)
-		.Title(Title)
-		.ClientSize( FMainFrameModule::GetProjectBrowserWindowSize() )
-		.SizingRule( ESizingRule::UserSized )
-		.SupportsMinimize(false) .SupportsMaximize(false);
+	TSharedRef<SWindow> NewProjectWindow = CreateProjectBrowserWindow(bAllowProjectOpening, bAllowProjectCreate);
 
 	NewProjectWindow->SetContent( FGameProjectGenerationModule::Get().CreateGameProjectDialog(bAllowProjectOpening, bAllowProjectCreate) );
 
@@ -827,6 +808,30 @@ void FMainFrameActionCallbacks::OpenWidgetReflector_Execute()
 	FGlobalTabmanager::Get()->TryInvokeTab(FTabId("WidgetReflector"));
 }
 
+TSharedRef<SWindow> FMainFrameActionCallbacks::CreateProjectBrowserWindow(bool bAllowProjectOpening, bool bAllowProjectCreate)
+{
+	FText Title;
+	if (bAllowProjectOpening && bAllowProjectCreate)
+	{
+		Title = LOCTEXT("SelectProjectWindowHeader", "Select Project");
+	}
+	else if (bAllowProjectOpening)
+	{
+		Title = LOCTEXT("OpenProjectWindowHeader", "Open Project");
+	}
+	else
+	{
+		Title = LOCTEXT("NewProjectWindowHeader", "New Project");
+	}
+
+	return SNew(SWindow)
+		.Title(Title)
+		.ClientSize(FMainFrameModule::GetProjectBrowserWindowSize())
+		.SizingRule(ESizingRule::UserSized)
+		.SupportsMinimize(false)
+		.SupportsMaximize(false);
+}
+
 
 /* FMainFrameActionCallbacks implementation
  *****************************************************************************/
diff --git a/Engine/Source/Editor/MainFrame/Private/Frame/MainFrameActions.h b/Engine/Source/Editor/MainFrame/Private/Frame/MainFrameActions.h
--- a/Engine/Source/Editor/MainFrame/Private/Frame/MainFrameActions.h
+++ b/Engine/Source/Editor/MainFrame/Private/Frame/MainFrameActions.h
@@ -10,6 +10,7 @@
 
 class FUICommandList;
 struct FRecentProjectFile;
+class SWindow;
 
 /**
  * Unreal editor main frame actions
@@ -208,6 +209,9 @@ public:
 
 	static void OpenWidgetReflector_Execute();
 
+	/** Return the window used for the Open and New project Browser */
+	static TSharedRef<SWindow> CreateProjectBrowserWindow(bool bAllowProjectOpening, bool bAllowProjectCreate);
+
 	/** Opens the new project dialog */
 	static void NewProject( bool bAllowProjectOpening, bool bAllowProjectCreate );
 
diff --git a/Engine/Source/Editor/MainFrame/Private/HomeScreen/HomeScreenMenuContext.h b/Engine/Source/Editor/MainFrame/Private/HomeScreen/HomeScreenMenuContext.h
new file mode 100644
--- /dev/null
+++ b/Engine/Source/Editor/MainFrame/Private/HomeScreen/HomeScreenMenuContext.h
@@ -0,0 +1,17 @@
+﻿// Copyright Epic Games, Inc. All Rights Reserved.
+
+#pragma once
+#include "UObject/Object.h"
+#include "HomeScreenMenuContext.generated.h"
+
+class SHomeScreen;
+
+UCLASS(MinimalAPI)
+class UHomeScreenContext : public UObject
+{
+	GENERATED_BODY()
+
+public:
+
+	TWeakPtr<SHomeScreen> HomeScreen;
+};
diff --git a/Engine/Source/Editor/MainFrame/Private/HomeScreen/HomeScreenWeb.cpp b/Engine/Source/Editor/MainFrame/Private/HomeScreen/HomeScreenWeb.cpp
--- a/Engine/Source/Editor/MainFrame/Private/HomeScreen/HomeScreenWeb.cpp
+++ b/Engine/Source/Editor/MainFrame/Private/HomeScreen/HomeScreenWeb.cpp
@@ -1,4 +1,4 @@
-﻿// Copyright Epic Games, Inc. All Rights Reserved.
+// Copyright Epic Games, Inc. All Rights Reserved.
 
 #include "HomeScreenWeb.h"
 #include "Internationalization/Regex.h"
@@ -13,6 +13,11 @@ void UHomeScreenWeb::NavigateTo(EMainSectionMenu InSectionToNavigate)
 	OnNavigationChangedDelegate.Broadcast(InSectionToNavigate);
 }
 
+void UHomeScreenWeb::OpenGettingStartedProject()
+{
+	OnTutorialProjectRequestedDelegate.Broadcast();
+}
+
 void UHomeScreenWeb::OpenWebPage(const FString& InURL) const
 {
 	const FText ErrorText = FText::Format(LOCTEXT("LaunchingURLNotAllowed", "Given URL is not registered as an allowed one:\n {0}"), FText::FromString(InURL));
diff --git a/Engine/Source/Editor/MainFrame/Private/HomeScreen/HomeScreenWeb.h b/Engine/Source/Editor/MainFrame/Private/HomeScreen/HomeScreenWeb.h
--- a/Engine/Source/Editor/MainFrame/Private/HomeScreen/HomeScreenWeb.h
+++ b/Engine/Source/Editor/MainFrame/Private/HomeScreen/HomeScreenWeb.h
@@ -6,6 +6,7 @@
 #include "HomeScreenWeb.generated.h"
 
 DECLARE_MULTICAST_DELEGATE_OneParam(FOnNavigationChanged, const EMainSectionMenu);
+DECLARE_MULTICAST_DELEGATE(FOnGettingStartedProjectRequested);
 
 UCLASS()
 class UHomeScreenWeb : public UObject
@@ -17,10 +18,15 @@ public:
 	void NavigateTo(EMainSectionMenu InSectionToNavigate);
 	FOnNavigationChanged& OnNavigationChanged() { return OnNavigationChangedDelegate; }
 
+	UFUNCTION()
+	void OpenGettingStartedProject();
+	FOnGettingStartedProjectRequested& OnTutorialProjectRequested() { return OnTutorialProjectRequestedDelegate; }
+
 	UFUNCTION()
 	void OpenWebPage(const FString& InURL) const;
 
 private:
 	FOnNavigationChanged OnNavigationChangedDelegate;
+	FOnGettingStartedProjectRequested OnTutorialProjectRequestedDelegate;
 	EMainSectionMenu SectionToNavigate = EMainSectionMenu::None;
 };
diff --git a/Engine/Source/Editor/MainFrame/Private/HomeScreen/SHomeScreen.cpp b/Engine/Source/Editor/MainFrame/Private/HomeScreen/SHomeScreen.cpp
--- a/Engine/Source/Editor/MainFrame/Private/HomeScreen/SHomeScreen.cpp
+++ b/Engine/Source/Editor/MainFrame/Private/HomeScreen/SHomeScreen.cpp
@@ -3,20 +3,26 @@
 #include "SHomeScreen.h"
 
 #include "DesktopPlatformModule.h"
+#include "Dialogs/SOutputLogDialog.h"
 #include "Editor.h"
+#include "Frame/MainFrameActions.h"
 #include "Frame/RootWindowLocation.h"
 #include "GameProjectGenerationModule.h"
 #include "GameProjectUtils.h"
+#include "HomeScreenMenuContext.h"
 #include "HomeScreenWeb.h"
 #include "HttpModule.h"
 #include "HttpRetrySystem.h"
 #include "IWebBrowserWindow.h"
-#include "Misc/StringOutputDevice.h"
+#include "MainFrameModule.h"
+#include "ProjectDescriptor.h"
 #include "Settings/EditorSettings.h"
 #include "Styling/AppStyle.h"
 #include "Styling/StyleColors.h"
 #include "SWebBrowser.h"
 #include "TimerManager.h"
+#include "ToolMenus.h"
+#include "UnrealEdGlobals.h"
 #include "Widgets/Images/SImage.h"
 #include "Widgets/Images/SThrobber.h"
 #include "Widgets/Input/SButton.h"
@@ -24,7 +30,6 @@
 #include "Widgets/Input/SComboButton.h"
 #include "Widgets/Layout/SBorder.h"
 #include "Widgets/Layout/SBox.h"
-#include "Widgets/Layout/SSeparator.h"
 #include "Widgets/Layout/SSpacer.h"
 #include "Widgets/Layout/SWidgetSwitcher.h"
 #include "Widgets/SBoxPanel.h"
@@ -94,13 +99,18 @@ void SHomeScreen::Construct(const FArguments& InArgs, const TSharedPtr<SWebBrows
 	if (WebObject.IsValid())
 	{
 		WebObject->OnNavigationChanged().AddSP(this, &SHomeScreen::OnNavigateToSection);
+		WebObject->OnTutorialProjectRequested().AddSP(this, &SHomeScreen::OnOpenGettingStartedProject);
 	}
 
 	AutoLoadProjectComboBoxSelection = EAutoLoadProject::HomeScreen;
-	MainHomeSelection = EMainSectionMenu::Home;
+	bHasLatestEngineProjects = HasAlreadyLatestEngineProject();
+	MainHomeSelection = bHasLatestEngineProjects ? EMainSectionMenu::Home : EMainSectionMenu::GettingStarted;
 
 	LoadSettings();
 
+	// Fixed Button Size
+	constexpr float SectionButtonCheckBoxHeight = 28.f;
+
 	// Main Section CheckBoxes
 	CreateMainSectionCheckBox(HomeCheckBox, EMainSectionMenu::Home, LOCTEXT("ProjectHomeHome", "Home"), FAppStyle::GetBrush("HomeScreen.Home"));
 	CreateMainSectionCheckBox(NewsCheckBox, EMainSectionMenu::News, LOCTEXT("ProjectHomeNews", "News"), FAppStyle::GetBrush("HomeScreen.News"));
@@ -141,7 +151,8 @@ void SHomeScreen::Construct(const FArguments& InArgs, const TSharedPtr<SWebBrows
 	ChildSlot
 	[
 		SNew(SBorder)
-		.BorderImage(FAppStyle::GetBrush("HomeScreen.Background"))
+		.BorderImage(FAppStyle::GetBrush("WhiteBrush"))
+		.BorderBackgroundColor(COLOR("#101014FF"))
 		.Padding(0.f)
 		[
 			SNew(SHorizontalBox)
@@ -168,23 +179,26 @@ void SHomeScreen::Construct(const FArguments& InArgs, const TSharedPtr<SWebBrows
 
 					// Create and MyFolder Buttons
 					+ SVerticalBox::Slot()
-					.AutoHeight()
+					.MinHeight(SectionButtonCheckBoxHeight)
+					.MaxHeight(SectionButtonCheckBoxHeight)
 					.Padding(8.f, 0.f, 8.f, 8.f)
 					[
 						SNew(SButton)
 						.ButtonStyle(FAppStyle::Get(), "HomeScreen.CreateNewProjectButton")
 						.OnClicked(this, &SHomeScreen::OnCreateProjectDialog, /** bInAllowProjectOpening */false, /** bInAllowProjectCreation */true)
 						.HAlign(HAlign_Center)
+						.ContentPadding(0.f)
 						[
 							SNew(SHorizontalBox)
 
 							+ SHorizontalBox::Slot()
 							.AutoWidth()
+							.VAlign(VAlign_Center)
 							.Padding(0.f, 0.f, 8.f, 0.f)
 							[
-								SNew(SImage)
-								.Image(FAppStyle::GetBrush("HomeScreen.PlusCircle"))
-								.ColorAndOpacity(FStyleColors::Black)
+									SNew(SImage)
+									.Image(FAppStyle::GetBrush("HomeScreen.PlusCircle"))
+									.ColorAndOpacity(FStyleColors::Black)
 							]
 
 							+ SHorizontalBox::Slot()
@@ -192,24 +206,28 @@ void SHomeScreen::Construct(const FArguments& InArgs, const TSharedPtr<SWebBrows
 							.VAlign(VAlign_Center)
 							[
 								SNew(STextBlock)
+								.Font(FAppStyle::GetFontStyle("HomeScreen.Font.NewAndCreateButton"))
 								.Text(LOCTEXT("ProjectHomeCreateNewProject", "New Project"))
 							]
 						]
 					]
 
 					+ SVerticalBox::Slot()
-					.AutoHeight()
+					.MinHeight(SectionButtonCheckBoxHeight)
+					.MaxHeight(SectionButtonCheckBoxHeight)
 					.Padding(8.f, 0.f, 8.f, 0.f)
 					[
 						SNew(SButton)
 						.ButtonStyle(FAppStyle::Get(), "HomeScreen.MyFolderButton")
 						.OnClicked(this, &SHomeScreen::OnCreateProjectDialog, /** bInAllowProjectOpening */true, /** bInAllowProjectCreation */false)
 						.HAlign(HAlign_Center)
+						.ContentPadding(0.f)
 						[
 							SNew(SHorizontalBox)
 
 							+ SHorizontalBox::Slot()
 							.AutoWidth()
+							.VAlign(VAlign_Center)
 							.Padding(0.f, 0.f, 8.f, 0.f)
 							[
 								SNew(SImage)
@@ -221,6 +239,7 @@ void SHomeScreen::Construct(const FArguments& InArgs, const TSharedPtr<SWebBrows
 							.VAlign(VAlign_Center)
 							[
 								SNew(STextBlock)
+								.Font(FAppStyle::GetFontStyle("HomeScreen.Font.NewAndCreateButton"))
 								.Text(LOCTEXT("ProjectHomeMyProjects", "My Projects"))
 							]
 						]
@@ -234,25 +253,29 @@ void SHomeScreen::Construct(const FArguments& InArgs, const TSharedPtr<SWebBrows
 						SNew(SVerticalBox)
 
 						+ SVerticalBox::Slot()
-						.AutoHeight()
+						.MinHeight(SectionButtonCheckBoxHeight)
+						.MaxHeight(SectionButtonCheckBoxHeight)
 						[
 							HomeCheckBox.ToSharedRef()
 						]
 					
 						+ SVerticalBox::Slot()
-						.AutoHeight()
+						.MinHeight(SectionButtonCheckBoxHeight)
+						.MaxHeight(SectionButtonCheckBoxHeight)
 						[
 							NewsCheckBox.ToSharedRef()
 						]
 
 						+ SVerticalBox::Slot()
-						.AutoHeight()
+						.MinHeight(SectionButtonCheckBoxHeight)
+						.MaxHeight(SectionButtonCheckBoxHeight)
 						[
 							GettingStartedCheckBox.ToSharedRef()
 						]
 
 						+ SVerticalBox::Slot()
-						.AutoHeight()
+						.MinHeight(SectionButtonCheckBoxHeight)
+						.MaxHeight(SectionButtonCheckBoxHeight)
 						[
 							SampleProjectsCheckBox.ToSharedRef()
 						]
@@ -265,7 +288,7 @@ void SHomeScreen::Construct(const FArguments& InArgs, const TSharedPtr<SWebBrows
 					[
 						SNew(STextBlock)
 						.Text(LOCTEXT("ResourceAndCommunity", "Resources & Community"))
-						.Font(FAppStyle::GetFontStyle("BoldFont"))
+						.Font(FAppStyle::GetFontStyle("HomeScreen.Font.DefaultBoldSize"))
 						.ColorAndOpacity(FStyleColors::White)
 					]
 
@@ -276,31 +299,36 @@ void SHomeScreen::Construct(const FArguments& InArgs, const TSharedPtr<SWebBrows
 						SNew(SVerticalBox)
 
 						+ SVerticalBox::Slot()
-						.AutoHeight()
+						.MinHeight(SectionButtonCheckBoxHeight)
+						.MaxHeight(SectionButtonCheckBoxHeight)
 						[
 							ForumsButton.ToSharedRef()
 						]
 
 						+ SVerticalBox::Slot()
-						.AutoHeight()
+						.MinHeight(SectionButtonCheckBoxHeight)
+						.MaxHeight(SectionButtonCheckBoxHeight)
 						[
 							DocumentationButton.ToSharedRef()
 						]
 
 						+ SVerticalBox::Slot()
-						.AutoHeight()
+						.MinHeight(SectionButtonCheckBoxHeight)
+						.MaxHeight(SectionButtonCheckBoxHeight)
 						[
 							TutorialsButton.ToSharedRef()
 						]
 
 						+ SVerticalBox::Slot()
-						.AutoHeight()
+						.MinHeight(SectionButtonCheckBoxHeight)
+						.MaxHeight(SectionButtonCheckBoxHeight)
 						[
 							RoadmapButton.ToSharedRef()
 						]
 
 						+ SVerticalBox::Slot()
-						.AutoHeight()
+						.MinHeight(SectionButtonCheckBoxHeight)
+						.MaxHeight(SectionButtonCheckBoxHeight)
 						[
 							ReleaseNotesButton.ToSharedRef()
 						]
@@ -360,7 +388,7 @@ void SHomeScreen::Construct(const FArguments& InArgs, const TSharedPtr<SWebBrows
 								.Padding(0.f, 0.f, 0.f, 16.f)
 								[
 									SNew(STextBlock)
-									.Font(FAppStyle::GetFontStyle("BoldFont"))
+									.Font(FAppStyle::GetFontStyle("HomeScreen.Font.DefaultBoldSize"))
 									.Text(LOCTEXT("HomeScreenNoInternet", "No Internet Connection"))
 									.ColorAndOpacity(FStyleColors::White)
 								]
@@ -371,8 +399,14 @@ void SHomeScreen::Construct(const FArguments& InArgs, const TSharedPtr<SWebBrows
 									SNew(SButton)
 									.HAlign(HAlign_Center)
 									.ButtonStyle(FAppStyle::Get(), "HomeScreen.MyFolderButton")
-									.Text(LOCTEXT("HomeScreenReconnect", "Reconnect"))
 									.OnClicked(this, &SHomeScreen::OnInternetConnectionRetried)
+									.HAlign(HAlign_Center)
+									.VAlign(VAlign_Center)
+									[
+										SNew(STextBlock)
+										.Font(FAppStyle::GetFontStyle("HomeScreen.Font.DefaultSize"))
+										.Text(LOCTEXT("HomeScreenReconnect", "Reconnect"))
+									]
 								]
 							]
 						]
@@ -383,6 +417,7 @@ void SHomeScreen::Construct(const FArguments& InArgs, const TSharedPtr<SWebBrows
 						.HAlign(HAlign_Left)
 						[
 							SNew(STextBlock)
+							.Font(FAppStyle::GetFontStyle("HomeScreen.Font.DefaultSize"))
 							.Text(LOCTEXT("ProjectHomeLoad", "Load on Startup"))
 						]
 
@@ -394,9 +429,11 @@ void SHomeScreen::Construct(const FArguments& InArgs, const TSharedPtr<SWebBrows
 							.ComboButtonStyle(&FAppStyle::Get().GetWidgetStyle<FComboButtonStyle>("HomeScreen.ComboButton"))
 							.MenuPlacement(EMenuPlacement::MenuPlacement_CenteredAboveAnchor)
 							.HasDownArrow(true)
+							.ContentPadding(FMargin(0.f ,2.f))
 							.ButtonContent()
 							[
 								SNew(STextBlock)
+								.Font(FAppStyle::GetFontStyle("HomeScreen.Font.DefaultSize"))
 								.Text(this, &SHomeScreen::GetAutoLoadProjectComboBoxLabelText)
 								.ColorAndOpacity_Lambda([this] ()
 								{
@@ -416,7 +453,7 @@ void SHomeScreen::Construct(const FArguments& InArgs, const TSharedPtr<SWebBrows
 						// Social Media Section
 						+ SVerticalBox::Slot()
 						.AutoHeight()
-						.Padding(46.f, 0.f)
+						.Padding(46.f, 0.f, 46.f, 20.f)
 						[
 							SNew(SHorizontalBox)
 
@@ -466,7 +503,7 @@ void SHomeScreen::Construct(const FArguments& InArgs, const TSharedPtr<SWebBrows
 				SNew(SWidgetSwitcher)
 				.WidgetIndex_Lambda([this] ()
 				{
-					return MainHomeSelection == EMainSectionMenu::Home ? 0 : 1;
+					return MainHomeSelection == EMainSectionMenu::Home && bHasLatestEngineProjects ? /** Show Recent Projects */ 0 : /** Show just the WebPage */ 1;
 				})
 
 				+ SWidgetSwitcher::Slot()
@@ -474,6 +511,7 @@ void SHomeScreen::Construct(const FArguments& InArgs, const TSharedPtr<SWebBrows
 					SNew(SVerticalBox)
 
 					+ SVerticalBox::Slot()
+					.Padding(0.f, 0.f, 0.f, 24.f)
 					.AutoHeight()
 					[
 						SNew(SHorizontalBox)
@@ -500,7 +538,7 @@ void SHomeScreen::Construct(const FArguments& InArgs, const TSharedPtr<SWebBrows
 								[
 									SNew(STextBlock)
 									.Text(LOCTEXT("RecentProjects", "Recent Projects"))
-									.Font(FAppStyle::GetFontStyle("HomeScreen.RecentProjectFont"))
+									.Font(FAppStyle::GetFontStyle("HomeScreen.Font.RecentProject"))
 									.ColorAndOpacity(FStyleColors::White)
 								]
 							
@@ -511,7 +549,12 @@ void SHomeScreen::Construct(const FArguments& InArgs, const TSharedPtr<SWebBrows
 									SNew(SButton)
 									.ButtonStyle(FAppStyle::Get(), "HomeScreen.SeeAllProjectsButton")
 									.OnClicked(this, &SHomeScreen::OnCreateProjectDialog, /** bInAllowProjectOpening */true, /** bInAllowProjectCreation */false)
-									.Text(LOCTEXT("SeeAllProjects", "See All Projects"))
+									.ContentPadding(FMargin(18, 10))
+									[
+										SNew(STextBlock)
+										.Font(FAppStyle::GetFontStyle("HomeScreen.Font.DefaultSize"))
+										.Text(LOCTEXT("SeeAllProjects", "See All Projects"))
+									]
 								]
 							]
 
@@ -521,7 +564,7 @@ void SHomeScreen::Construct(const FArguments& InArgs, const TSharedPtr<SWebBrows
 								SNew(SBorder)
 								.Clipping(EWidgetClipping::ClipToBounds)
 								.BorderImage(FAppStyle::GetBrush("HomeScreen.RecentProjects.Background"))
-								.Padding(16.f)
+								.Padding(24.f)
 								[
 									SNew(SVerticalBox)
 
@@ -531,7 +574,7 @@ void SHomeScreen::Construct(const FArguments& InArgs, const TSharedPtr<SWebBrows
 										.HeightOverride(172.f)
 										.Padding(0.f)
 										[
-												ProjectBrowser.ToSharedRef()
+											ProjectBrowser.ToSharedRef()
 										]
 									]
 								]
@@ -565,27 +608,13 @@ void SHomeScreen::Construct(const FArguments& InArgs, const TSharedPtr<SWebBrows
 
 FReply SHomeScreen::OnCreateProjectDialog(bool bInAllowProjectOpening, bool bInAllowProjectCreation)
 {
-	FRootWindowLocation DefaultWindowLocation;
-
-	FText Title = LOCTEXT("ProjectBrowserTitle", "Unreal Project Browser");
-	FText ToolTip = LOCTEXT("ProjectBrowserTitleToolTip", "Close Unreal Project Browser");
-	if (bInAllowProjectCreation)
+	if (GUnrealEd->WarnIfLightingBuildIsCurrentlyRunning())
 	{
-		Title = LOCTEXT("ProjectCreationTitle", "Unreal Project Creation");
-		ToolTip = LOCTEXT("ProjectCreationTitleToolTip", "Close Unreal Project Creation");
+		return FReply::Handled();
 	}
 
-	const TSharedRef<SWindow> NewProjectWindow = SNew(SWindow)
-		.AutoCenter(EAutoCenter::PreferredWorkArea)
-		.Title(Title)
-		.CloseButtonToolTipText(ToolTip)
-		.IsInitiallyMaximized(false)
-		.ScreenPosition(DefaultWindowLocation.ScreenPosition)
-		.ClientSize(DefaultWindowLocation.WindowSize)
-		.CreateTitleBar(true)
-		.SizingRule(ESizingRule::UserSized)
-		.SupportsMaximize(false)
-		.SupportsMinimize(false);
+	TSharedRef<SWindow> NewProjectWindow =
+		FMainFrameActionCallbacks::CreateProjectBrowserWindow(bInAllowProjectOpening, bInAllowProjectCreation);
 
 	NewProjectWindow->SetContent(FGameProjectGenerationModule::Get().CreateGameProjectDialog(bInAllowProjectOpening, bInAllowProjectCreation));
 	FSlateApplication::Get().AddModalWindow(NewProjectWindow, nullptr);
@@ -611,20 +640,21 @@ void SHomeScreen::OnMainHomeSectionChanged(ECheckBoxState InCheckBoxState, EMain
 	{
 	case EMainSectionMenu::News:
 		MainHomeSelection = EMainSectionMenu::News;
-		URL = TEXT("https://www.unrealengine.com/en-US/feed");
+		URL = TEXT("https://editor-website.ol.epicgames.net/en-US/news");
 		break;
 	case EMainSectionMenu::GettingStarted:
 		MainHomeSelection = EMainSectionMenu::GettingStarted;
-		URL = TEXT("https://dev.epicgames.com/community/unreal-engine/getting-started");
+		URL = TEXT("https://editor-website.ol.epicgames.net/en-US/get-started-with-unreal-engine");
 		break;
 	case EMainSectionMenu::SampleProjects:
 		MainHomeSelection = EMainSectionMenu::SampleProjects;
-		URL = TEXT("https://dev.epicgames.com/community/unreal-engine/learning");
+		URL = TEXT("https://editor-website.ol.epicgames.net/en-US/samples");
 		break;
 	case EMainSectionMenu::None:
 	case EMainSectionMenu::Home:
 	default:
 		MainHomeSelection = EMainSectionMenu::Home;
+		URL = bHasLatestEngineProjects ? TEXT("https://editor-website.ol.epicgames.net/en-US/recent-projects") : TEXT("https://editor-website.ol.epicgames.net/en-US/no-projects");
 		URL = TEXT("https://editor-website.ol.epicgames.net/en-US");
 		break;
 	}
@@ -645,22 +675,20 @@ void SHomeScreen::LoadSettings()
 
 		AutoLoadProjectComboBoxSelection = (EAutoLoadProject)LoadingType;
 	}
-
-	bool bShowHomePage;
-	if (GConfig->GetBool(TEXT("HomeScreen"), TEXT("ShowHomePage"), bShowHomePage, GEngineIni))
-	{
-		bShowHomePageOption = bShowHomePage;
-	}
 }
 
 void SHomeScreen::SaveSettings()
 {
 	GConfig->SetInt(TEXT("HomeScreen"), TEXT("Loading"), (int32)AutoLoadProjectComboBoxSelection, GEngineIni);
-	GConfig->SetBool(TEXT("HomeScreen"), TEXT("ShowHomePage"), bShowHomePageOption, GEngineIni);
 }
 
-FReply SHomeScreen::OnAutoloadLastProjectChanged(EAutoLoadProject InAutoLoadOption)
+FReply SHomeScreen::OnAutoLoadOptionChanged(EAutoLoadProject InAutoLoadOption)
 {
+	if (AutoLoadProjectComboBoxSelection == InAutoLoadOption)
+	{
+		return FReply::Handled();
+	}
+
 	UEditorSettings *Settings = GetMutableDefault<UEditorSettings>();
 	AutoLoadProjectComboBoxSelection = InAutoLoadOption;
 
@@ -674,15 +702,18 @@ FReply SHomeScreen::OnAutoloadLastProjectChanged(EAutoLoadProject InAutoLoadOpti
 		Settings->PostEditChangeProperty(PropertyUpdateStruct);
 	}
 
+	if (TSharedPtr<SWidget> ComboButtonMenu = ComboButtonMenuWeak.Pin())
+	{
+		FSlateApplication::Get().DismissMenuByWidget(ComboButtonMenu.ToSharedRef());
+	}
+
 	SaveSettings();
 	return FReply::Handled();
 }
 
-FReply SHomeScreen::OnShowHomePageWhenNewAnnouncementOptionClicked()
+EVisibility SHomeScreen::IsAutoLoadOptionCheckVisible(EAutoLoadProject InAutoLoadOption) const
 {
-	bShowHomePageOption = !bShowHomePageOption;
-	SaveSettings();
-	return FReply::Handled();
+	return AutoLoadProjectComboBoxSelection == InAutoLoadOption ? EVisibility::Visible : EVisibility::Hidden;
 }
 
 bool SHomeScreen::IsCheckBoxCheckedOrHovered(const TSharedPtr<SCheckBox> InCheckBox) const
@@ -719,7 +750,8 @@ void SHomeScreen::CreateMainSectionCheckBox(TSharedPtr<SCheckBox>& OutCheckBox,
 
 		+ SHorizontalBox::Slot()
 		.AutoWidth()
-		.Padding(0.f, 0.f, 4.f, 0.f)
+		.VAlign(VAlign_Center)
+		.Padding(0.f, 0.f, 8.f, 0.f)
 		[
 			SNew(SImage)
 			.Image(InImage)
@@ -731,6 +763,7 @@ void SHomeScreen::CreateMainSectionCheckBox(TSharedPtr<SCheckBox>& OutCheckBox,
 		.VAlign(VAlign_Center)
 		[
 			SNew(STextBlock)
+			.Font(FAppStyle::GetFontStyle("HomeScreen.Font.DefaultSize"))
 			.Text(InText)
 			.ColorAndOpacity(this, &SHomeScreen::GetMainSectionCheckBoxColor, OutCheckBox)
 		]);
@@ -749,7 +782,8 @@ void SHomeScreen::CreateResourceButtons(TSharedPtr<SButton>& OutButton, FString
 
 		+ SHorizontalBox::Slot()
 		.AutoWidth()
-		.Padding(0.f, 0.f, 4.f, 0.f)
+		.VAlign(VAlign_Center)
+		.Padding(0.f, 0.f, 8.f, 0.f)
 		[
 			SNew(SImage)
 			.Image(InImage)
@@ -761,6 +795,7 @@ void SHomeScreen::CreateResourceButtons(TSharedPtr<SButton>& OutButton, FString
 		.VAlign(VAlign_Center)
 		[
 			SNew(STextBlock)
+			.Font(FAppStyle::GetFontStyle("HomeScreen.Font.DefaultSize"))
 			.Text(InText)
 			.ColorAndOpacity(this, &SHomeScreen::GetResourceAndSocialMediaButtonColor, OutButton)
 		]);
@@ -780,118 +815,20 @@ void SHomeScreen::CreateSocialMediaButtons(TSharedPtr<SButton>& OutButton, FStri
 
 TSharedRef<SWidget> SHomeScreen::CreateComboButtonMenuContentWidget()
 {
-	return SNew(SVerticalBox)
-
-	+ SVerticalBox::Slot()
-	.AutoHeight()
-	.Padding(8.f, 8.f)
-	[
-		SNew(SButton)
-		.ButtonStyle(FAppStyle::Get(), "HomeScreen.ComboButton.MenuButton")
-		.ContentPadding(FMargin(4.f, 4.f))
-		.OnClicked(this, &SHomeScreen::OnAutoloadLastProjectChanged, EAutoLoadProject::HomeScreen)
-		[
-			SNew(SHorizontalBox)
-
-			+ SHorizontalBox::Slot()
-			.AutoWidth()
-			[
-				SNew(SImage)
-				.Image(FAppStyle::GetBrush("Icons.Check"))
-				.Visibility(this, &SHomeScreen::GetAutoLoadProjectComboBoxContentIconVisibility, EAutoLoadProject::HomeScreen)
-			]
-
-			+ SHorizontalBox::Slot()
-			.AutoWidth()
-			[
-				SNew(STextBlock)
-				.Text(LOCTEXT("HomeScreenHomeTab", "Home Tab"))
-				.ColorAndOpacity(FStyleColors::White)
-			]
-		]
-	]
-
-	+ SVerticalBox::Slot()
-	.AutoHeight()
-	.Padding(8.f, 0.f)
-	[
-		SNew(SButton)
-		.ButtonStyle(FAppStyle::Get(), "HomeScreen.ComboButton.MenuButton")
-		.ContentPadding(FMargin(4.f, 4.f))
-		.OnClicked(this, &SHomeScreen::OnAutoloadLastProjectChanged, EAutoLoadProject::LastProject)
-		[
-			SNew(SHorizontalBox)
-
-			+ SHorizontalBox::Slot()
-			.AutoWidth()
-			[
-				SNew(SImage)
-				.Image(FAppStyle::GetBrush("Icons.Check"))
-				.Visibility(this, &SHomeScreen::GetAutoLoadProjectComboBoxContentIconVisibility, EAutoLoadProject::LastProject)
-
-			]
-
-			+ SHorizontalBox::Slot()
-			.AutoWidth()
-			[
-				SNew(STextBlock)
-				.Text(LOCTEXT("HomeScreenMostRecentProject", "Most Recent Project"))
-				.ColorAndOpacity(FStyleColors::White)
-			]
-		]
-	]
+	UToolMenus* ToolMenus = UToolMenus::Get();
+	if (ToolMenus->IsMenuRegistered("HomeScreenComboButtonMenu"))
+	{
+		FToolMenuContext HomeScreenComboContext;
 
-	+ SVerticalBox::Slot()
-	.AutoHeight()
-	.Padding(8.f, 8.f)
-	[
-		SNew(SSeparator)
-		.Orientation(Orient_Horizontal)
-		.Thickness(1.f)
-		.SeparatorImage(FAppStyle::GetBrush("WhiteBrush"))
-		.ColorAndOpacity(FStyleColors::White25)
-	]
-
-	+ SVerticalBox::Slot()
-	.AutoHeight()
-	.Padding(8.f, 0.f, 8.f, 8.f)
-	[
-		SNew(SButton)
-		.ButtonStyle(FAppStyle::Get(), "HomeScreen.ComboButton.MenuButton")
-		.ContentPadding(FMargin(4.f, 4.f))
-		.OnClicked(this, &SHomeScreen::OnShowHomePageWhenNewAnnouncementOptionClicked)
-		[
-			SNew(SHorizontalBox)
+		UHomeScreenContext* HomeScreenContextObject = NewObject<UHomeScreenContext>();
+		HomeScreenContextObject->HomeScreen = SharedThis(this);
 
-			+ SHorizontalBox::Slot()
-			.AutoWidth()
-			.VAlign(VAlign_Center)
-			[
-				SNew(SBox)
-				.HeightOverride(16.f)
-				[
-					SNew(SImage)
-					.Image(FAppStyle::GetBrush("Icons.Check"))
-					.Visibility_Lambda([this] ()
-					{
-						return bShowHomePageOption ? EVisibility::Visible : EVisibility::Hidden;
-					})
-				]
-			]
-
-			+ SHorizontalBox::Slot()
-			.AutoWidth()
-			[
-				SNew(STextBlock)
-				.WrapTextAt(222.f)
-				.Text(LOCTEXT("HomeScreenShowHomePage", "Show Home page when new announcements are available"))
-				.ColorAndOpacity_Lambda([this] ()
-				{
-					return bShowHomePageOption ? FStyleColors::White : FStyleColors::White25;
-				})
-			]
-		]
-	];
+		HomeScreenComboContext.AddObject(HomeScreenContextObject);
+		TSharedRef<SWidget> ComboButtonMenuRef = ToolMenus->GenerateWidget("HomeScreenComboButtonMenu", HomeScreenComboContext);
+		ComboButtonMenuWeak = ComboButtonMenuRef;
+		return ComboButtonMenuRef;
+	}
+	return SNullWidget::NullWidget;
 }
 
 void SHomeScreen::CheckInternetConnection()
@@ -946,14 +883,107 @@ void SHomeScreen::OnNavigateToSection(EMainSectionMenu InSectionToNavigate)
 	OnMainHomeSectionChanged(ECheckBoxState::Checked, InSectionToNavigate);
 }
 
-FText SHomeScreen::GetAutoLoadProjectComboBoxLabelText() const
+void SHomeScreen::OnOpenGettingStartedProject()
 {
-	return AutoLoadProjectComboBoxSelection == EAutoLoadProject::HomeScreen ? LOCTEXT("HomeScreenComboHomeTab", "Home Tab") : LOCTEXT("HomeScreenComboMostRecentProject", "Most Recent Project");
+	FString TemplateRootFolders = FPaths::RootDir() + TEXT("Templates");
+	FString UEIntroRootFolder = TemplateRootFolders / TEXT("TP_UEIntro_BP");
+	const FString SearchString = UEIntroRootFolder / TEXT("*.") + FProjectDescriptor::GetExtension();
+	TArray<FString> FoundProjectFiles;
+	IFileManager::Get().FindFiles(FoundProjectFiles, *SearchString, /*Files=*/true, /*Directories=*/false);
+
+	// There should be only 1 match
+	if (FoundProjectFiles.Num() != 1)
+	{
+		return;
+	}
+
+	const FString TemplateProjectPath = UEIntroRootFolder / FoundProjectFiles[0];
+	const FString DesiredProjectName = TEXT("UEIntroProject");
+
+	// Get the default project creation folder
+	IDesktopPlatform* DesktopPlatform = FDesktopPlatformModule::Get();
+	if (!DesktopPlatform)
+	{
+		return;
+	}
+
+	FString DefaultProjectFolder = DesktopPlatform->GetDefaultProjectCreationPath();
+
+	// Find a unique name
+	FString ProjectName = DesiredProjectName;
+	FString NewProjectFolder = FPaths::Combine(DefaultProjectFolder, ProjectName);
+	int32 Counter = 1;
+	constexpr int32 MaxTries = 1000;
+
+	while (FPaths::DirectoryExists(NewProjectFolder))
+	{
+		ProjectName = FString::Printf(TEXT("%s_%d"), *DesiredProjectName, Counter++);
+		NewProjectFolder = FPaths::Combine(DefaultProjectFolder, ProjectName);
+		if (Counter > MaxTries)
+		{
+			FMessageDialog::Open(EAppMsgType::Ok, LOCTEXT("TooManyTriesToFindAUniqueName", "Something went wrong when trying to find a unique name for the tutorial project.\nCheck your folder containing the projects and clean that up of unneeded tutorial projects."));
+			return;
+		}
+	}
+
+	const FString ProjectPath = DefaultProjectFolder;
+	const FString Filename = ProjectName + TEXT(".") + FProjectDescriptor::GetExtension();
+	FString ProjectFilename = FPaths::Combine(*ProjectPath, *ProjectName, *Filename);
+	FPaths::MakePlatformFilename(ProjectFilename);
+
+	FText FailReason, FailLog;
+	FProjectInformation ProjectInfo;
+	ProjectInfo.TemplateCategory = TEXT("Game");
+	ProjectInfo.TemplateFile = TemplateProjectPath;
+	ProjectInfo.ProjectFilename = ProjectFilename;
+
+	if (!FGameProjectGenerationModule::Get().CreateProject(ProjectInfo, FailReason, FailLog))
+	{
+		SOutputLogDialog::Open(LOCTEXT("CreateProject", "Create Project"), FailReason, FailLog, FText::GetEmpty());
+		return;
+	}
+
+	// Successfully created the project. Update the last created location string.
+	FString CreatedProjectPath = FPaths::GetPath(FPaths::GetPath(ProjectFilename));
+
+	// If the original path was the drives root (ie: C:/) the double path call strips the last /
+	if (CreatedProjectPath.EndsWith(":"))
+	{
+		CreatedProjectPath.AppendChar('/');
+	}
+
+	UEditorSettings* Settings = GetMutableDefault<UEditorSettings>();
+	Settings->CreatedProjectPaths.Remove(CreatedProjectPath);
+	Settings->CreatedProjectPaths.Insert(CreatedProjectPath, 0);
+	Settings->PostEditChange();
+
+	// Open the project
+	FText FailReasonOpen;
+	if (FGameProjectGenerationModule::Get().OpenProject(ProjectFilename, FailReasonOpen))
+	{
+		// Successfully opened the project, the editor is closing.
+		// Close this window in case something prevents the editor from closing (save dialog, quit confirmation, etc)
+		if (FApp::HasProjectName())
+		{
+			TSharedPtr<SWindow> ContainingWindow = FSlateApplication::Get().FindWidgetWindow(AsShared());
+	
+			if (ContainingWindow.IsValid())
+			{
+				ContainingWindow->RequestDestroyWindow();
+			}
+		}
+	}
+	else
+	{
+		FString ErrorString = FailReasonOpen.ToString();
+		UE_LOG(LogTemp, Log, TEXT("%s"), *ErrorString);
+		FMessageDialog::Open(EAppMsgType::Ok, FailReasonOpen);
+	}
 }
 
-EVisibility SHomeScreen::GetAutoLoadProjectComboBoxContentIconVisibility(EAutoLoadProject InAutoLoadOption) const
+FText SHomeScreen::GetAutoLoadProjectComboBoxLabelText() const
 {
-	return AutoLoadProjectComboBoxSelection == InAutoLoadOption ? EVisibility::Visible : EVisibility::Hidden;
+	return AutoLoadProjectComboBoxSelection == EAutoLoadProject::HomeScreen ? LOCTEXT("HomeScreenComboHomeTab", "Home Tab") : LOCTEXT("HomeScreenComboMostRecentProject", "Most Recent Project");
 }
 
 FReply SHomeScreen::OnInternetConnectionRetried()
@@ -962,4 +992,166 @@ FReply SHomeScreen::OnInternetConnectionRetried()
 	return FReply::Handled();
 }
 
+bool SHomeScreen::HasAlreadyLatestEngineProject() const
+{
+	const FString EngineIdentifier = FDesktopPlatformModule::Get()->GetCurrentEngineIdentifier();
+	TArray<FString> ProjectFiles;
+
+	FString RootDir;
+	IDesktopPlatform* DesktopPlatformModule = FDesktopPlatformModule::Get();
+	if (!DesktopPlatformModule)
+	{
+		return false;
+	}
+
+	DesktopPlatformModule->EnumerateProjectsKnownByEngine(EngineIdentifier, false, ProjectFiles);
+	DesktopPlatformModule->GetEngineRootDirFromIdentifier(EngineIdentifier, RootDir);
+	ProjectFiles.RemoveAll([RootDir, DesktopPlatformModule] (const FString& InProjectFile)
+		{
+			FString ProjectFilename = IFileManager::Get().ConvertToAbsolutePathForExternalAppForRead(*InProjectFile);
+			if (!FPaths::FileExists(ProjectFilename))
+			{
+				return true;
+			}
+
+			if (ProjectFilename.Contains(RootDir))
+			{
+				return true;
+			}
+			FString Identifier;
+			DesktopPlatformModule->GetEngineIdentifierForProject(ProjectFilename, Identifier);
+
+			FEngineVersion ProjectEngineVersion;
+
+			if (DesktopPlatformModule->IsStockEngineRelease(Identifier))
+			{
+				DesktopPlatformModule->TryParseStockEngineVersion(Identifier, ProjectEngineVersion);
+			}
+
+			if (ProjectEngineVersion.IsEmpty())
+			{
+				FString RootProjectDir;
+				DesktopPlatformModule->GetEngineRootDirFromIdentifier(Identifier, RootProjectDir);
+				if (!DesktopPlatformModule->TryGetEngineVersion(RootProjectDir, ProjectEngineVersion))
+				{
+					return true;
+				}
+			}
+
+			const FString EngineVersionString = FEngineVersion::Current().ToString(EVersionComponent::Patch);
+			const FString ProjectEngineVersionString = ProjectEngineVersion.ToString(EVersionComponent::Patch);
+
+			if (EngineVersionString != ProjectEngineVersionString)
+			{
+				return true;
+			}
+
+			return false;
+		});
+
+	return !ProjectFiles.IsEmpty();
+}
+
+FDelayedAutoRegisterHelper SHomeScreen::LoadStartupComboButtonRegistration(
+	EDelayedRegisterRunPhase::EndOfEngineInit,
+	[]
+		{
+			FToolMenuOwnerScoped OwnerScoped(UE_MODULE_NAME);
+			static const FName LoadStartupComboButtonName("HomeScreenComboButtonMenu");
+			UToolMenu* Menu = UToolMenus::Get()->RegisterMenu(LoadStartupComboButtonName);
+			FToolMenuSection& Section = Menu->AddSection("StartupSection");
+
+			Section.AddDynamicEntry(
+				TEXT("HomeTabDynamic"),
+				FNewToolMenuSectionDelegate::CreateLambda([] (FToolMenuSection& InSection)
+					{
+						if (UHomeScreenContext* Context = InSection.FindContext<UHomeScreenContext>())
+						{
+							if (const TSharedPtr<SHomeScreen> HomeScreen = Context->HomeScreen.Pin())
+							{
+								const TSharedRef<SHomeScreen> HomeScreenRef = HomeScreen.ToSharedRef();
+
+								TSharedRef<SButton> HomeTabWidget = SNew(SButton)
+									.ButtonStyle(FAppStyle::Get(), "HomeScreen.ComboButton.MenuButton")
+									.ContentPadding(FMargin(12.f, 10.f))
+									.OnClicked(HomeScreenRef, &SHomeScreen::OnAutoLoadOptionChanged, EAutoLoadProject::HomeScreen)
+									[
+										SNew(SHorizontalBox)
+
+										+ SHorizontalBox::Slot()
+										.Padding(0.f, 0.f, 8.f, 0.f)
+										.AutoWidth()
+										[
+											SNew(SImage)
+											.Image(FAppStyle::GetBrush("Icons.Check"))
+											.Visibility(HomeScreenRef, &SHomeScreen::IsAutoLoadOptionCheckVisible, EAutoLoadProject::HomeScreen)
+										]
+
+										+ SHorizontalBox::Slot()
+										.AutoWidth()
+										[
+											SNew(STextBlock)
+											.Text(LOCTEXT("HomeScreenHomeTab", "Home Tab"))
+											.Font(FAppStyle::GetFontStyle("HomeScreen.Font.DefaultSize"))
+											.ColorAndOpacity(FStyleColors::White)
+										]
+									];
+
+								InSection.AddEntry(FToolMenuEntry::InitWidget(
+									TEXT("HomeTab"),
+									SNew(SBox).Padding(8.f, 0.f, 8.f, 4.f)[HomeTabWidget],
+									FText::GetEmpty()));
+							}
+						}
+					})
+				);
+
+
+			Section.AddDynamicEntry(
+				TEXT("MostRecentProjectDynamic"),
+				FNewToolMenuSectionDelegate::CreateLambda([] (FToolMenuSection& InSection)
+					{
+						if (UHomeScreenContext* Context = InSection.FindContext<UHomeScreenContext>())
+						{
+							if (const TSharedPtr<SHomeScreen> HomeScreen = Context->HomeScreen.Pin())
+							{
+								const TSharedRef<SHomeScreen> HomeScreenRef = HomeScreen.ToSharedRef();
+
+								TSharedRef<SButton> MostRecentProjectWidget = SNew(SButton)
+									.ButtonStyle(FAppStyle::Get(), "HomeScreen.ComboButton.MenuButton")
+									.ContentPadding(FMargin(12.f, 10.f))
+									.OnClicked(HomeScreenRef, &SHomeScreen::OnAutoLoadOptionChanged, EAutoLoadProject::LastProject)
+									[
+										SNew(SHorizontalBox)
+
+										+ SHorizontalBox::Slot()
+										.Padding(0.f, 0.f, 8.f, 0.f)
+										.AutoWidth()
+										[
+											SNew(SImage)
+											.Image(FAppStyle::GetBrush("Icons.Check"))
+											.Visibility(HomeScreenRef, &SHomeScreen::IsAutoLoadOptionCheckVisible, EAutoLoadProject::LastProject)
+										]
+
+										+ SHorizontalBox::Slot()
+										.AutoWidth()
+										[
+											SNew(STextBlock)
+											.Text(LOCTEXT("HomeScreenMostRecentProject", "Most Recent Project"))
+											.Font(FAppStyle::GetFontStyle("HomeScreen.Font.DefaultSize"))
+											.ColorAndOpacity(FStyleColors::White)
+										]
+									];
+
+								InSection.AddEntry(FToolMenuEntry::InitWidget(
+									TEXT("MostRecentProject"),
+									SNew(SBox).Padding(8.f, 0.f)[MostRecentProjectWidget],
+									FText::GetEmpty()));
+							}
+						}
+					})
+				);
+		}
+);
+
 #undef LOCTEXT_NAMESPACE
diff --git a/Engine/Source/Editor/MainFrame/Private/HomeScreen/SHomeScreen.h b/Engine/Source/Editor/MainFrame/Private/HomeScreen/SHomeScreen.h
--- a/Engine/Source/Editor/MainFrame/Private/HomeScreen/SHomeScreen.h
+++ b/Engine/Source/Editor/MainFrame/Private/HomeScreen/SHomeScreen.h
@@ -57,10 +57,10 @@ private:
 	void SaveSettings();
 
 	/** Handles changes to autoload last project selection */
-	FReply OnAutoloadLastProjectChanged(EAutoLoadProject InAutoLoadOption);
+	FReply OnAutoLoadOptionChanged(EAutoLoadProject InAutoLoadOption);
 
-	/** Handles clicks for showing the home page when new announcements are available */
-	FReply OnShowHomePageWhenNewAnnouncementOptionClicked();
+	/** True if the current AutoLoad option match the given InAutoLoadOption */
+	EVisibility IsAutoLoadOptionCheckVisible(EAutoLoadProject InAutoLoadOption) const;
 
 	/** Checks if a checkbox is checked or hovered */
 	bool IsCheckBoxCheckedOrHovered(const TSharedPtr<SCheckBox> InCheckBox) const;
@@ -98,15 +98,22 @@ private:
 	/** Handles navigation to a section from a web request */
 	void OnNavigateToSection(EMainSectionMenu InSectionToNavigate);
 
+	/** Handles the getting started template project creation */
+	void OnOpenGettingStartedProject();
+
 	/** Gets the label text for the autoload project combo box */
 	FText GetAutoLoadProjectComboBoxLabelText() const;
 
-	/** Get the visibility for the check icon beside the ComboBox options */
-	EVisibility GetAutoLoadProjectComboBoxContentIconVisibility(EAutoLoadProject InAutoLoadOption) const;
-
 	/** Executed when clicking on the Reconnect button when no internet connection is detected */
 	FReply OnInternetConnectionRetried();
 
+	/** Whether the current user already created project with this Engine version */
+	bool HasAlreadyLatestEngineProject() const;
+
+private:
+	/** Registers the LoadStartup ComboButton Menu */
+	static FDelayedAutoRegisterHelper LoadStartupComboButtonRegistration;
+
 private:
 	/** Currently selected main home section */
 	EMainSectionMenu MainHomeSelection = EMainSectionMenu::Home;
@@ -134,21 +141,24 @@ private:
 	/** Widget for project browser */
 	TSharedPtr<SWidget> ProjectBrowser;
 
+	/** Current combo button menu */
+	TWeakPtr<SWidget> ComboButtonMenuWeak;
+
 	/** Timer handle for checking internet connection */
 	FTimerHandle CheckInternetConnectionTimerHandle;
 
 	/** Path of the currently selected project */
 	FString CurrentSelectedProjectPath;
 
-	/** Whether to show the home page option */
-	bool bShowHomePageOption = false;
-
 	/** Whether the system is connected to the internet */
 	bool bIsConnected = true;
 
 	/** Whether the last request has finished */
 	bool bIsRequestFinished = true;
 
+	/** Whether the current user already created project with this Engine version */
+	bool bHasLatestEngineProjects = false;
+
 	/** How many times we executed the function without being connected */
 	int32 TimerManagerCountOnceDisconnected = 0;
 
diff --git a/Engine/Source/Editor/MainFrame/Private/MainFrameModule.cpp b/Engine/Source/Editor/MainFrame/Private/MainFrameModule.cpp
--- a/Engine/Source/Editor/MainFrame/Private/MainFrameModule.cpp
+++ b/Engine/Source/Editor/MainFrame/Private/MainFrameModule.cpp
@@ -983,6 +983,29 @@ TSharedRef<SDockTab> FMainFrameModule::SpawnHomeScreenWindow(const FSpawnTabArgs
 		WindowSettings.BrowserFrameRate = 60;
 
 		WebBrowserWindow = WebBrowserSingleton->CreateBrowserWindow(WindowSettings);
+
+		WebBrowserSingleton->SetDevToolsShortcutEnabled(true);
+		WebBrowserWindow->OnCreateWindow().BindLambda([](const TWeakPtr<IWebBrowserWindow>& NewBrowserWindow, const TWeakPtr<IWebBrowserPopupFeatures>& PopupFeatures)
+		{
+			// Initialize a dialog
+			auto DialogMainWindow = SNew(SWindow)
+				.Title(FText::FromString(TEXT("Chrome Debugging Tools")))
+				.ClientSize(FVector2D(700, 700))
+				.SupportsMaximize(true)
+				.SupportsMinimize(true)
+				[
+					SNew(SVerticalBox)
+
+					+ SVerticalBox::Slot()
+					.HAlign(HAlign_Fill)
+					.VAlign(VAlign_Fill)
+					[
+						SNew(SWebBrowser, NewBrowserWindow.Pin())
+					]
+				];
+			FSlateApplication::Get().AddWindow(DialogMainWindow);
+			return true;
+		});
 	}
 
 	if (!WebBrowser.IsValid())
@@ -990,7 +1013,6 @@ TSharedRef<SDockTab> FMainFrameModule::SpawnHomeScreenWindow(const FSpawnTabArgs
 		SAssignNew(WebBrowser, SWebBrowser, WebBrowserWindow)
 		.ShowControls(false)
 		.SupportsTransparency(true)
-		.ShowControls(false)
 		.ShowAddressBar(false);
 
 		constexpr bool bIsPermanent = true;
