[HomeScreen] Small code fixes

Removed testing url line.
Moved assignement of EngineIdentifier after the DesktopPlatform module invalid check.

diff --git a/Engine/Source/Editor/MainFrame/Private/HomeScreen/SHomeScreen.cpp b/Engine/Source/Editor/MainFrame/Private/HomeScreen/SHomeScreen.cpp
--- a/Engine/Source/Editor/MainFrame/Private/HomeScreen/SHomeScreen.cpp
+++ b/Engine/Source/Editor/MainFrame/Private/HomeScreen/SHomeScreen.cpp
@@ -655,7 +655,6 @@ void SHomeScreen::OnMainHomeSectionChanged(ECheckBoxState InCheckBoxState, EMain
 	default:
 		MainHomeSelection = EMainSectionMenu::Home;
 		URL = bHasLatestEngineProjects ? TEXT("https://editor-website.ol.epicgames.net/en-US/recent-projects") : TEXT("https://editor-website.ol.epicgames.net/en-US/no-projects");
-		URL = TEXT("https://editor-website.ol.epicgames.net/en-US");
 		break;
 	}
 
@@ -994,7 +993,6 @@ FReply SHomeScreen::OnInternetConnectionRetried()
 
 bool SHomeScreen::HasAlreadyLatestEngineProject() const
 {
-	const FString EngineIdentifier = FDesktopPlatformModule::Get()->GetCurrentEngineIdentifier();
 	TArray<FString> ProjectFiles;
 
 	FString RootDir;
@@ -1004,6 +1002,8 @@ bool SHomeScreen::HasAlreadyLatestEngineProject() const
 		return false;
 	}
 
+	const FString EngineIdentifier = DesktopPlatformModule->GetCurrentEngineIdentifier();
+
 	DesktopPlatformModule->EnumerateProjectsKnownByEngine(EngineIdentifier, false, ProjectFiles);
 	DesktopPlatformModule->GetEngineRootDirFromIdentifier(EngineIdentifier, RootDir);
 	ProjectFiles.RemoveAll([RootDir, DesktopPlatformModule] (const FString& InProjectFile)
