[HomeScreen] Small Fix

If unable to retrieve the engine RootDir fake that it's a new user.
Added AutoHeight attribute to VerticalBox to make alignment correct.

diff --git a/Engine/Source/Editor/MainFrame/Private/HomeScreen/SHomeScreen.cpp b/Engine/Source/Editor/MainFrame/Private/HomeScreen/SHomeScreen.cpp
--- a/Engine/Source/Editor/MainFrame/Private/HomeScreen/SHomeScreen.cpp
+++ b/Engine/Source/Editor/MainFrame/Private/HomeScreen/SHomeScreen.cpp
@@ -179,6 +179,7 @@ void SHomeScreen::Construct(const FArguments& InArgs, const TSharedPtr<SWebBrows
 
 					// Create and MyFolder Buttons
 					+ SVerticalBox::Slot()
+					.AutoHeight()
 					.MinHeight(SectionButtonCheckBoxHeight)
 					.MaxHeight(SectionButtonCheckBoxHeight)
 					.Padding(8.f, 0.f, 8.f, 8.f)
@@ -213,6 +214,7 @@ void SHomeScreen::Construct(const FArguments& InArgs, const TSharedPtr<SWebBrows
 					]
 
 					+ SVerticalBox::Slot()
+					.AutoHeight()
 					.MinHeight(SectionButtonCheckBoxHeight)
 					.MaxHeight(SectionButtonCheckBoxHeight)
 					.Padding(8.f, 0.f, 8.f, 0.f)
@@ -253,6 +255,7 @@ void SHomeScreen::Construct(const FArguments& InArgs, const TSharedPtr<SWebBrows
 						SNew(SVerticalBox)
 
 						+ SVerticalBox::Slot()
+						.AutoHeight()
 						.MinHeight(SectionButtonCheckBoxHeight)
 						.MaxHeight(SectionButtonCheckBoxHeight)
 						[
@@ -260,6 +263,7 @@ void SHomeScreen::Construct(const FArguments& InArgs, const TSharedPtr<SWebBrows
 						]
 					
 						+ SVerticalBox::Slot()
+						.AutoHeight()
 						.MinHeight(SectionButtonCheckBoxHeight)
 						.MaxHeight(SectionButtonCheckBoxHeight)
 						[
@@ -267,6 +271,7 @@ void SHomeScreen::Construct(const FArguments& InArgs, const TSharedPtr<SWebBrows
 						]
 
 						+ SVerticalBox::Slot()
+						.AutoHeight()
 						.MinHeight(SectionButtonCheckBoxHeight)
 						.MaxHeight(SectionButtonCheckBoxHeight)
 						[
@@ -274,6 +279,7 @@ void SHomeScreen::Construct(const FArguments& InArgs, const TSharedPtr<SWebBrows
 						]
 
 						+ SVerticalBox::Slot()
+						.AutoHeight()
 						.MinHeight(SectionButtonCheckBoxHeight)
 						.MaxHeight(SectionButtonCheckBoxHeight)
 						[
@@ -299,6 +305,7 @@ void SHomeScreen::Construct(const FArguments& InArgs, const TSharedPtr<SWebBrows
 						SNew(SVerticalBox)
 
 						+ SVerticalBox::Slot()
+						.AutoHeight()
 						.MinHeight(SectionButtonCheckBoxHeight)
 						.MaxHeight(SectionButtonCheckBoxHeight)
 						[
@@ -306,6 +313,7 @@ void SHomeScreen::Construct(const FArguments& InArgs, const TSharedPtr<SWebBrows
 						]
 
 						+ SVerticalBox::Slot()
+						.AutoHeight()
 						.MinHeight(SectionButtonCheckBoxHeight)
 						.MaxHeight(SectionButtonCheckBoxHeight)
 						[
@@ -313,6 +321,7 @@ void SHomeScreen::Construct(const FArguments& InArgs, const TSharedPtr<SWebBrows
 						]
 
 						+ SVerticalBox::Slot()
+						.AutoHeight()
 						.MinHeight(SectionButtonCheckBoxHeight)
 						.MaxHeight(SectionButtonCheckBoxHeight)
 						[
@@ -320,6 +329,7 @@ void SHomeScreen::Construct(const FArguments& InArgs, const TSharedPtr<SWebBrows
 						]
 
 						+ SVerticalBox::Slot()
+						.AutoHeight()
 						.MinHeight(SectionButtonCheckBoxHeight)
 						.MaxHeight(SectionButtonCheckBoxHeight)
 						[
@@ -327,6 +337,7 @@ void SHomeScreen::Construct(const FArguments& InArgs, const TSharedPtr<SWebBrows
 						]
 
 						+ SVerticalBox::Slot()
+						.AutoHeight()
 						.MinHeight(SectionButtonCheckBoxHeight)
 						.MaxHeight(SectionButtonCheckBoxHeight)
 						[
@@ -1005,7 +1016,12 @@ bool SHomeScreen::HasAlreadyLatestEngineProject() const
 	const FString EngineIdentifier = DesktopPlatformModule->GetCurrentEngineIdentifier();
 
 	DesktopPlatformModule->EnumerateProjectsKnownByEngine(EngineIdentifier, false, ProjectFiles);
-	DesktopPlatformModule->GetEngineRootDirFromIdentifier(EngineIdentifier, RootDir);
+	if (!DesktopPlatformModule->GetEngineRootDirFromIdentifier(EngineIdentifier, RootDir))
+	{
+		// The engine root dir couldn't be found so fake that it's a new user
+		return false;
+	}
+
 	ProjectFiles.RemoveAll([RootDir, DesktopPlatformModule] (const FString& InProjectFile)
 		{
 			FString ProjectFilename = IFileManager::Get().ConvertToAbsolutePathForExternalAppForRead(*InProjectFile);
