[HomeScreen] Dropdown menu highlights when highlighting other major tab dropdown menus

diff --git a/Engine/Source/Editor/MainFrame/Private/MainFrameModule.cpp b/Engine/Source/Editor/MainFrame/Private/MainFrameModule.cpp
--- a/Engine/Source/Editor/MainFrame/Private/MainFrameModule.cpp
+++ b/Engine/Source/Editor/MainFrame/Private/MainFrameModule.cpp
@@ -962,6 +962,15 @@ TSharedRef<SDockTab> FMainFrameModule::SpawnHomeScreenWindow(const FSpawnTabArgs
 		.IconColor(FStyleColors::White.GetSpecifiedColor())
 		.ContentPadding(FMargin(0));
 
+	if (!HomeScreenTabManager.IsValid())
+	{
+		HomeScreenTabManager = FGlobalTabmanager::Get()->NewTabManager(HomeScreenTab);
+		HomeScreenTabManager->SetAllowWindowMenuBar(true);
+		FToolMenuContext EmptyContext;
+		MakeMainMenu(HomeScreenTabManager, "MainFrame.NomadMainMenu", EmptyContext);
+		HomeScreenTab->SetOnTabClosed(SDockTab::FOnTabClosedCallback::CreateLambda([this] (TSharedRef<SDockTab>) { HomeScreenTabManager.Reset(); }));
+	}
+
 	if (!WebObject.IsValid())
 	{
 		WebObject = TStrongObjectPtr(NewObject<UHomeScreenWeb>());
diff --git a/Engine/Source/Editor/MainFrame/Private/MainFrameModule.h b/Engine/Source/Editor/MainFrame/Private/MainFrameModule.h
--- a/Engine/Source/Editor/MainFrame/Private/MainFrameModule.h
+++ b/Engine/Source/Editor/MainFrame/Private/MainFrameModule.h
@@ -322,4 +322,8 @@ private:
 
 	// Web browser widget instance
 	static TSharedPtr<SWebBrowser> WebBrowser;
+
+	/** Holds the tab manager that manages the HomeScreen tab. */
+    TSharedPtr<FTabManager> HomeScreenTabManager;
+
 };
