[HomeScreen] Remove part of the menu entry in the WindowMenu bar at the top

#FailedValidations Preflight

diff --git a/Engine/Source/Editor/MainFrame/Private/MainFrameModule.cpp b/Engine/Source/Editor/MainFrame/Private/MainFrameModule.cpp
--- a/Engine/Source/Editor/MainFrame/Private/MainFrameModule.cpp
+++ b/Engine/Source/Editor/MainFrame/Private/MainFrameModule.cpp
@@ -41,6 +41,8 @@
 #include "HomeScreen/SHomeScreen.h" 
 #include "IWebBrowserSingleton.h"
 #include "IWebBrowserWindow.h"
+#include "Menus/SettingsMenu.h"
+#include "Settings/EditorStyleSettings.h"
 #include "SWebBrowser.h"
 #include "ToolMenus.h"
 #include "WebBrowserModule.h"
@@ -967,7 +969,7 @@ TSharedRef<SDockTab> FMainFrameModule::SpawnHomeScreenWindow(const FSpawnTabArgs
 		HomeScreenTabManager = FGlobalTabmanager::Get()->NewTabManager(HomeScreenTab);
 		HomeScreenTabManager->SetAllowWindowMenuBar(true);
 		FToolMenuContext EmptyContext;
-		MakeMainMenu(HomeScreenTabManager, "MainFrame.NomadMainMenu", EmptyContext);
+		MakeMainMenu(HomeScreenTabManager, "HomeScreen.MainMenu", EmptyContext);
 		HomeScreenTab->SetOnTabClosed(SDockTab::FOnTabClosedCallback::CreateLambda([this] (TSharedRef<SDockTab>) { HomeScreenTabManager.Reset(); }));
 	}
 
@@ -1043,7 +1045,133 @@ void FMainFrameModule::HandleCodeAccessorOpenFileFailed(const FString& Filename)
 	FSlateNotificationManager::Get().QueueNotification(Info);
 }
 
-#undef LOCTEXT_NAMESPACE
+void FMainFrameModule::RegisterHomeScreenFileMenu()
+{
+	UToolMenus* ToolMenus = UToolMenus::Get();
+	UToolMenu* FileMenu = ToolMenus->RegisterMenu("HomeScreen.MainMenu.File");
+
+	FToolMenuSection& ProjectSection = FileMenu->AddSection("FileProject", LOCTEXT("ProjectHeading", "Project"));
+	ProjectSection.AddMenuEntry(FMainFrameCommands::Get().NewProject);
+	ProjectSection.AddMenuEntry(FMainFrameCommands::Get().OpenProject);
+
+#if !PLATFORM_MAC // Handled by app's menu in menu bar
+	FToolMenuSection& ExitSection = FileMenu->AddSection("Exit", LOCTEXT("Exit", "Exit"), FToolMenuInsert("FileProject", EToolMenuInsertType::After));
+	ExitSection.AddSeparator("Exit");
+	ExitSection.AddMenuEntry(FMainFrameCommands::Get().Exit);
+#endif
+}
+
+void FMainFrameModule::RegisterHomeScreenEditMenu()
+{
+	UToolMenu* EditMenu = UToolMenus::Get()->RegisterMenu("HomeScreen.MainMenu.Edit");
+	FToolMenuSection& Section = EditMenu->AddSection("Configuration", LOCTEXT("ConfigurationHeading", "Configuration"));
+		if (GetDefault<UEditorStyleSettings>()->bExpandConfigurationMenus)
+		{
+			Section.AddSubMenu(
+				"EditorPreferencesSubMenu",
+				LOCTEXT("EditorPreferencesSubMenuLabel", "Editor Preferences"),
+				LOCTEXT("EditorPreferencesSubMenuToolTip", "Configure the behavior and features of this Editor"),
+				FNewToolMenuDelegate::CreateStatic(&FSettingsMenu::MakeMenu, FName("Editor")),
+				false,
+				FSlateIcon(FAppStyle::GetAppStyleSetName(), "EditorPreferences.TabIcon")
+			);
+
+			Section.AddSubMenu(
+				"ProjectSettingsSubMenu",
+				LOCTEXT("ProjectSettingsSubMenuLabel", "Project Settings"),
+				LOCTEXT("ProjectSettingsSubMenuToolTip", "Change the settings of the currently loaded project"),
+				FNewToolMenuDelegate::CreateStatic(&FSettingsMenu::MakeMenu, FName("Project")),
+				false,
+				FSlateIcon(FAppStyle::GetAppStyleSetName(), "ProjectSettings.TabIcon")
+			);
+		}
+		else
+		{
+#if !PLATFORM_MAC // Handled by app's menu in menu bar
+			{
+				IMainFrameModule& MainFrame = FModuleManager::LoadModuleChecked<IMainFrameModule>("MainFrame");
+				FName CategoryName;
+				FName SectionName;
+				MainFrame.GetEditorSettingsDefaultSelectionOverride(CategoryName, SectionName);
+
+				if (CategoryName.IsNone())
+				{
+					CategoryName = FName("General");
+				}
+				if (SectionName.IsNone())
+				{
+					SectionName = FName("Appearance");
+				}
+
+				Section.AddMenuEntry(
+					"EditorPreferencesMenu",
+					LOCTEXT("EditorPreferencesMenuLabel", "Editor Preferences..."),
+					LOCTEXT("EditorPreferencesMenuToolTip", "Configure the behavior and features of the Unreal Editor."),
+					FSlateIcon(FAppStyle::GetAppStyleSetName(), "EditorPreferences.TabIcon"),
+					FUIAction(FExecuteAction::CreateStatic(&FSettingsMenu::OpenSettings, FName("Editor"), CategoryName, SectionName))
+				);
+			}
+#endif //if !PLATFORM_MAC
+
+			Section.AddMenuEntry(
+				"ProjectSettingsMenu",
+				LOCTEXT("ProjectSettingsMenuLabel", "Project Settings..."),
+				LOCTEXT("ProjectSettingsMenuToolTip", "Change the settings of the currently loaded project."),
+				FSlateIcon(FAppStyle::GetAppStyleSetName(), "ProjectSettings.TabIcon"),
+				FUIAction(FExecuteAction::CreateStatic(&FSettingsMenu::OpenSettings, FName("Project"), FName("Project"), FName("General")))
+			);
+		}
+}
+
+void FMainFrameModule::RegisterHomeScreenHelpMenu()
+{
+	UToolMenus* ToolMenus = UToolMenus::Get();
+	ToolMenus->RegisterMenu("HomeScreen.MainMenu.Help", "MainFrame.MainMenu.Help", EMultiBoxType::Menu);
+}
+
+FDelayedAutoRegisterHelper FMainFrameModule::HomeScreenWindowMenuRegistration(
+	EDelayedRegisterRunPhase::EndOfEngineInit,
+	[]
+		{
+			UToolMenus* ToolMenus = UToolMenus::Get();
+			static const FName HomeScreenMainMenu("HomeScreen.MainMenu");
+			if (!ToolMenus->IsMenuRegistered(HomeScreenMainMenu))
+			{
+				UToolMenu* HomeScreenMenuBar = UToolMenus::Get()->RegisterMenu(HomeScreenMainMenu, NAME_None, EMultiBoxType::MenuBar);
 
+				RegisterHomeScreenFileMenu();
+				RegisterHomeScreenEditMenu();
+				RegisterHomeScreenHelpMenu();
+
+				HomeScreenMenuBar->StyleName = FName(TEXT("WindowMenuBar"));
+
+				HomeScreenMenuBar->AddSubMenu(
+					"HomeScreen",
+					NAME_None,
+					"File",
+					LOCTEXT("FileMenu", "File"),
+					LOCTEXT("FileMenu_ToolTip", "Open the file menu")
+				);
+
+				HomeScreenMenuBar->AddSubMenu(
+					"HomeScreen",
+					NAME_None,
+					"Edit",
+					LOCTEXT("EditMenu", "Edit"),
+					LOCTEXT("EditMenu_ToolTip", "Open the edit menu")
+					);
+
+				HomeScreenMenuBar->AddSubMenu(
+					"HomeScreen",
+					NAME_None,
+					"Help",
+					LOCTEXT("HelpMenu", "Help"),
+					LOCTEXT("HelpMenu_ToolTip", "Open the help menu")
+					);
+			}
+		}
+);
+
+#undef LOCTEXT_NAMESPACE
 
 IMPLEMENT_MODULE(FMainFrameModule, MainFrame);
diff --git a/Engine/Source/Editor/MainFrame/Private/MainFrameModule.h b/Engine/Source/Editor/MainFrame/Private/MainFrameModule.h
--- a/Engine/Source/Editor/MainFrame/Private/MainFrameModule.h
+++ b/Engine/Source/Editor/MainFrame/Private/MainFrameModule.h
@@ -255,6 +255,18 @@ private:
 	{
 		DelayedShowMainFrameDelegate.Unbind();
 	}
+
+	// Register HomeScreen file window menu
+	static void RegisterHomeScreenFileMenu();
+
+	// Register HomeScreen edit window menu
+	static void RegisterHomeScreenEditMenu();
+
+	// Register HomeScreen help window menu
+	static void RegisterHomeScreenHelpMenu();
+
+	// Register the Window Menu for the HomeScreen tab
+    static FDelayedAutoRegisterHelper HomeScreenWindowMenuRegistration;
 private:
 
 	// Weak pointer to the level editor's compile notification item.
