Tweak the project browser size following cl 46268800

diff --git a/Engine/Source/Editor/MainFrame/Private/MainFrameModule.h b/Engine/Source/Editor/MainFrame/Private/MainFrameModule.h
--- a/Engine/Source/Editor/MainFrame/Private/MainFrameModule.h
+++ b/Engine/Source/Editor/MainFrame/Private/MainFrameModule.h
@@ -222,7 +222,7 @@ public:
 public:
 
 	/** Get the size of the project browser window */
-	static FVector2D GetProjectBrowserWindowSize() { return FVector2D(1190, 822); }
+	static FVector2D GetProjectBrowserWindowSize() { return FVector2D(1206, 822); }
 
 private:
 
