[HomeScreen] Use CVar to change the engine load workflow from the ProjectBrowser to the HomeScreen

diff --git a/Engine/Source/Editor/MainFrame/Private/MainFrameModule.cpp b/Engine/Source/Editor/MainFrame/Private/MainFrameModule.cpp
--- a/Engine/Source/Editor/MainFrame/Private/MainFrameModule.cpp
+++ b/Engine/Source/Editor/MainFrame/Private/MainFrameModule.cpp
@@ -51,11 +51,10 @@
 DEFINE_LOG_CATEGORY(LogMainFrame);
 #define LOCTEXT_NAMESPACE "FMainFrameModule"
 
-#if WITH_EDITOR
 TStrongObjectPtr<UHomeScreenWeb> FMainFrameModule::WebObject = nullptr;
 TSharedPtr<IWebBrowserWindow> FMainFrameModule::WebBrowserWindow = nullptr;
 TSharedPtr<SWebBrowser> FMainFrameModule::WebBrowser = nullptr;
-#endif
+TSharedPtr<FTabManager> FMainFrameModule::HomeScreenTabManager = nullptr;
 
 namespace UE::Editor::HomeScreen
 {
@@ -149,6 +148,8 @@ bool FProjectDialogProvider::IsRequestingMainFrameControl() const
 
 FMainFrameWindowOverrides FProjectDialogProvider::GetDesiredWindowConfiguration() const
 {
+	bool bIsHomeScreenEnabled = UE::Editor::HomeScreen::IsHomeScreenEnabled();
+
 	FMainFrameWindowOverrides ProjectDialogWindowConfig;
 
 	// Force tabs restored from layout that have no window (the LevelEditor tab) to use a docking area with
@@ -157,9 +158,9 @@ FMainFrameWindowOverrides FProjectDialogProvider::GetDesiredWindowConfiguration(
 	ProjectDialogWindowConfig.bEmbedTitleAreaContent = false;
 
 	// Do not maximize the window initially. Keep a small dialog feel.
-	ProjectDialogWindowConfig.bInitiallyMaximized = false;
+	ProjectDialogWindowConfig.bInitiallyMaximized = bIsHomeScreenEnabled;
 
-	ProjectDialogWindowConfig.WindowSize = FMainFrameModule::GetProjectBrowserWindowSize();
+	ProjectDialogWindowConfig.WindowSize = bIsHomeScreenEnabled ? TOptional<FVector2D>() : FMainFrameModule::GetProjectBrowserWindowSize();
 
 	ProjectDialogWindowConfig.bIsUserSizable = true;
 	ProjectDialogWindowConfig.bSupportsMaximize = true;
@@ -167,12 +168,22 @@ FMainFrameWindowOverrides FProjectDialogProvider::GetDesiredWindowConfiguration(
 	ProjectDialogWindowConfig.CenterRules = EAutoCenter::PreferredWorkArea;
 	// When opening the project dialog, show "Project Browser" in the window title
 	ProjectDialogWindowConfig.WindowTitle = LOCTEXT("ProjectBrowserDialogTitle", "Unreal Project Browser");
+	if (bIsHomeScreenEnabled)
+	{
+		const FString EngineVersionString = FEngineVersion::Current().ToString(EVersionComponent::Minor);
+		ProjectDialogWindowConfig.WindowTitle = FText::Format(LOCTEXT("HomeScreenDialogTitle", "Unreal Engine {0}"), FText::FromString(EngineVersionString));;
+	}
 
 	return ProjectDialogWindowConfig;
 }
 
 TSharedRef<SWidget> FProjectDialogProvider::CreateMainFrameContentWidget() const
 {
+	if (UE::Editor::HomeScreen::IsHomeScreenEnabled())
+	{
+		return FMainFrameModule::CreateHomeScreenWidget();
+	}
+
 	return FGameProjectGenerationModule::Get().CreateGameProjectDialog(/*bAllowProjectOpening=*/true, /*bAllowProjectCreate=*/true);
 }
 
@@ -681,33 +692,33 @@ void FMainFrameModule::StartupModule( )
 	SourceCodeAccessModule.OnOpenFileFailed().AddRaw( this, &FMainFrameModule::HandleCodeAccessorOpenFileFailed );
 
 	// HomeScreen Registration
-		if (UE::Editor::HomeScreen::IsHomeScreenEnabled())
-    	{
-    		FGlobalTabmanager::Get()->RegisterTabSpawner(
-    			"HomeScreen", FOnSpawnTab::CreateRaw(this, &FMainFrameModule::SpawnHomeScreenWindow))
-    			.SetAutoGenerateMenuEntry(false)
-    			.SetMenuType(ETabSpawnerMenuType::Hidden)
-    			.SetDisplayName(FText::GetEmpty());
-    
-    		// Register menu extension
-    		UToolMenus::RegisterStartupCallback(
-    			FSimpleMulticastDelegate::FDelegate::CreateLambda([] ()
-    				{
-    					UToolMenu* Menu = UToolMenus::Get()->ExtendMenu("LevelEditor.MainMenu.Tools");
-    
-    					FToolMenuSection& Section = Menu->AddSection("HomeScreenSection", FText::FromString("Home Screen"));
-    					Section.AddMenuEntry(
-    						"HomeScreen",
-    						FText::FromString("Home Screen"),
-    						FText::FromString("Opens the Home Screen tab"),
-    						FSlateIcon(FAppStyle::GetAppStyleSetName(), "HomeScreen.Home"),
-    						FUIAction(FExecuteAction::CreateLambda([] ()
-    							{
-    								FGlobalTabmanager::Get()->TryInvokeTab(FTabId("HomeScreen"));
-    							}))
-    					);
-    				})
-    		);
+	if (UE::Editor::HomeScreen::IsHomeScreenEnabled())
+	{
+		FGlobalTabmanager::Get()->RegisterTabSpawner(
+			"HomeScreen", FOnSpawnTab::CreateStatic(&FMainFrameModule::SpawnHomeScreenWindow))
+			.SetAutoGenerateMenuEntry(false)
+			.SetMenuType(ETabSpawnerMenuType::Hidden)
+			.SetDisplayName(FText::GetEmpty());
+
+		// Register menu extension
+		UToolMenus::RegisterStartupCallback(
+			FSimpleMulticastDelegate::FDelegate::CreateLambda([] ()
+				{
+					UToolMenu* Menu = UToolMenus::Get()->ExtendMenu("LevelEditor.MainMenu.Tools");
+
+					FToolMenuSection& Section = Menu->AddSection("HomeScreenSection", FText::FromString("Home Screen"));
+					Section.AddMenuEntry(
+						"HomeScreen",
+						FText::FromString("Home Screen"),
+						FText::FromString("Opens the Home Screen tab"),
+						FSlateIcon(FAppStyle::GetAppStyleSetName(), "HomeScreen.Home"),
+						FUIAction(FExecuteAction::CreateLambda([] ()
+							{
+								FGlobalTabmanager::Get()->TryInvokeTab(FTabId("HomeScreen"));
+							}))
+					);
+				})
+			);
     	}
 
 	ModuleCompileStartTime = 0.0f;
@@ -968,9 +979,19 @@ TSharedRef<SDockTab> FMainFrameModule::SpawnHomeScreenWindow(const FSpawnTabArgs
 	HomeScreenTabManager = FGlobalTabmanager::Get()->NewTabManager(HomeScreenTab);
 	HomeScreenTabManager->SetAllowWindowMenuBar(true);
 	FToolMenuContext EmptyContext;
-	MakeMainMenu(HomeScreenTabManager, "HomeScreen.MainMenu", EmptyContext);
-	HomeScreenTab->SetOnTabClosed(SDockTab::FOnTabClosedCallback::CreateLambda([this] (TSharedRef<SDockTab>) { HomeScreenTabManager.Reset(); }));
+	FMainFrameModule::Get().MakeMainMenu(HomeScreenTabManager, "HomeScreen.MainMenu", EmptyContext);
+	HomeScreenTab->SetOnTabClosed(SDockTab::FOnTabClosedCallback::CreateStatic([] (TSharedRef<SDockTab>) { HomeScreenTabManager.Reset(); }));
+
+	HomeScreenTab->SetTabIcon(FAppStyle::Get().GetBrush("HomeScreen.Home"));
 
+	// This will create the necessary object and widgets needed for the HomeScreen and return a newly created one or NullWidget if something went wrong
+	HomeScreenTab->SetContent(FMainFrameModule::CreateHomeScreenWidget());
+
+	return HomeScreenTab;
+}
+
+TSharedRef<SWidget> FMainFrameModule::CreateHomeScreenWidget()
+{
 	if (!WebObject.IsValid())
 	{
 		WebObject = TStrongObjectPtr(NewObject<UHomeScreenWeb>());
@@ -982,8 +1003,9 @@ TSharedRef<SDockTab> FMainFrameModule::SpawnHomeScreenWindow(const FSpawnTabArgs
 		if (!IWebBrowserModule::IsAvailable() || !WebBrowserModule.IsWebModuleAvailable())
 		{
 			FMessageDialog::Open(EAppMsgType::Ok, LOCTEXT("WebBrowserModuleNotAvailable", "Failed to load the plugin. Please enable Web WebBrowserWindow in the plugin manager to use the Web Panels."));
-			return HomeScreenTab;
+			return SNullWidget::NullWidget;
 		}
+
 		IWebBrowserSingleton* WebBrowserSingleton = WebBrowserModule.GetSingleton();
 
 		FCreateBrowserWindowSettings WindowSettings;
@@ -1004,7 +1026,7 @@ TSharedRef<SDockTab> FMainFrameModule::SpawnHomeScreenWindow(const FSpawnTabArgs
 				.SupportsMinimize(true)
 				[
 					SNew(SVerticalBox)
-
+	
 					+ SVerticalBox::Slot()
 					.HAlign(HAlign_Fill)
 					.VAlign(VAlign_Fill)
@@ -1029,11 +1051,7 @@ TSharedRef<SDockTab> FMainFrameModule::SpawnHomeScreenWindow(const FSpawnTabArgs
 		WebBrowserWindow->Reload();
 	}
 
-	HomeScreenTab->SetTabIcon(FAppStyle::Get().GetBrush("HomeScreen.Home"));
-
-	HomeScreenTab->SetContent(SNew(SHomeScreen, WebBrowser, TWeakObjectPtr(WebObject.Get())));
-
-	return HomeScreenTab;
+	return SNew(SHomeScreen, WebBrowser, TWeakObjectPtr(WebObject.Get()));
 }
 
 void FMainFrameModule::HandleCodeAccessorOpenFileFailed(const FString& Filename)
diff --git a/Engine/Source/Editor/MainFrame/Private/MainFrameModule.h b/Engine/Source/Editor/MainFrame/Private/MainFrameModule.h
--- a/Engine/Source/Editor/MainFrame/Private/MainFrameModule.h
+++ b/Engine/Source/Editor/MainFrame/Private/MainFrameModule.h
@@ -224,6 +224,12 @@ public:
 	/** Get the size of the project browser window */
 	static FVector2D GetProjectBrowserWindowSize() { return FVector2D(1206, 822); }
 
+	// Spawns the home screen tab
+	static TSharedRef<SDockTab> SpawnHomeScreenWindow(const FSpawnTabArgs& InArgs);
+
+	// Create the home screen widget
+	static TSharedRef<SWidget> CreateHomeScreenWidget();
+
 private:
 
 	// Handles the level editor module starting to recompile.
@@ -247,9 +253,6 @@ private:
 	// Handles launching code accessor
 	void HandleCodeAccessorLaunching( );
 
-	// Spawns the home screen tab
-	TSharedRef<SDockTab> SpawnHomeScreenWindow(const FSpawnTabArgs& InArgs);
-
 	// Reset delegate
 	void ClearDelayedShowMainFrameDelegate()
 	{
@@ -335,7 +338,7 @@ private:
 	// Web browser widget instance
 	static TSharedPtr<SWebBrowser> WebBrowser;
 
-	/** Holds the tab manager that manages the HomeScreen tab. */
-    TSharedPtr<FTabManager> HomeScreenTabManager;
+	// Holds the tab manager that manages the HomeScreen tab.
+    static TSharedPtr<FTabManager> HomeScreenTabManager;
 
 };
