[HomeScreen] Using -game launch argument causes client to crash

Added callback to UToolMenus RegisterStartupCallback to FDelayedAutoRegister

diff --git a/Engine/Source/Editor/MainFrame/Private/MainFrameModule.cpp b/Engine/Source/Editor/MainFrame/Private/MainFrameModule.cpp
--- a/Engine/Source/Editor/MainFrame/Private/MainFrameModule.cpp
+++ b/Engine/Source/Editor/MainFrame/Private/MainFrameModule.cpp
@@ -1130,10 +1130,10 @@ void FMainFrameModule::RegisterHomeScreenHelpMenu()
 	ToolMenus->RegisterMenu("HomeScreen.MainMenu.Help", "MainFrame.MainMenu.Help", EMultiBoxType::Menu);
 }
 
-FDelayedAutoRegisterHelper FMainFrameModule::HomeScreenWindowMenuRegistration(
-	EDelayedRegisterRunPhase::EndOfEngineInit,
-	[]
-		{
+FDelayedAutoRegisterHelper FMainFrameModule::HomeScreenWindowMenuRegistration(EDelayedRegisterRunPhase::EndOfEngineInit, []{
+		UToolMenus::RegisterStartupCallback(FSimpleMulticastDelegate::FDelegate::CreateLambda([] ()
+			{
+			FToolMenuOwnerScoped OwnerScoped(UE_MODULE_NAME);
 			UToolMenus* ToolMenus = UToolMenus::Get();
 			static const FName HomeScreenMainMenu("HomeScreen.MainMenu");
 			if (!ToolMenus->IsMenuRegistered(HomeScreenMainMenu))
@@ -1170,7 +1170,8 @@ FDelayedAutoRegisterHelper FMainFrameModule::HomeScreenWindowMenuRegistration(
 					LOCTEXT("HelpMenu_ToolTip", "Open the help menu")
 					);
 			}
-		}
+		}));
+	}
 );
 
 #undef LOCTEXT_NAMESPACE
