[HomeScreen] Remove menu entry to open the HomeScreen

diff --git a/Engine/Source/Editor/MainFrame/Private/MainFrameModule.cpp b/Engine/Source/Editor/MainFrame/Private/MainFrameModule.cpp
--- a/Engine/Source/Editor/MainFrame/Private/MainFrameModule.cpp
+++ b/Engine/Source/Editor/MainFrame/Private/MainFrameModule.cpp
@@ -684,26 +684,6 @@ void FMainFrameModule::StartupModule( )
 			.SetAutoGenerateMenuEntry(false)
 			.SetMenuType(ETabSpawnerMenuType::Hidden)
 			.SetDisplayName(FText::GetEmpty());
-
-		// Register menu extension
-		UToolMenus::RegisterStartupCallback(
-			FSimpleMulticastDelegate::FDelegate::CreateLambda([] ()
-				{
-					UToolMenu* Menu = UToolMenus::Get()->ExtendMenu("LevelEditor.MainMenu.Tools");
-
-					FToolMenuSection& Section = Menu->AddSection("HomeScreenSection", FText::FromString("Home Screen"));
-					Section.AddMenuEntry(
-						"HomeScreen",
-						FText::FromString("Home Screen"),
-						FText::FromString("Opens the Home Screen tab"),
-						FSlateIcon(FAppStyle::GetAppStyleSetName(), "HomeScreen.Home"),
-						FUIAction(FExecuteAction::CreateLambda([] ()
-							{
-								FGlobalTabmanager::Get()->TryInvokeTab(FTabId("HomeScreen"));
-							}))
-					);
-				})
-			);
     	}
 
 	ModuleCompileStartTime = 0.0f;
