[HomeScreen] Add HomeScreen project settings

#FailedValidations Preflight

diff --git a/Engine/Source/Editor/GameProjectGeneration/Private/SProjectBrowser.cpp b/Engine/Source/Editor/GameProjectGeneration/Private/SProjectBrowser.cpp
--- a/Engine/Source/Editor/GameProjectGeneration/Private/SProjectBrowser.cpp
+++ b/Engine/Source/Editor/GameProjectGeneration/Private/SProjectBrowser.cpp
@@ -1264,7 +1264,10 @@ void SProjectBrowser::OnFilterTextChanged(const FText& InText)
 void SProjectBrowser::OnAutoloadLastProjectChanged(ECheckBoxState NewState)
 {
 	UEditorSettings *Settings = GetMutableDefault<UEditorSettings>();
+
+	PRAGMA_DISABLE_DEPRECATION_WARNINGS
 	Settings->bLoadTheMostRecentlyLoadedProjectAtStartup = (NewState == ECheckBoxState::Checked);
+	PRAGMA_ENABLE_DEPRECATION_WARNINGS
 
 	FProperty* AutoloadProjectProperty = FindFProperty<FProperty>(Settings->GetClass(), "bLoadTheMostRecentlyLoadedProjectAtStartup");
 	if (AutoloadProjectProperty != NULL)
diff --git a/Engine/Source/Editor/GameProjectGeneration/Private/SProjectDialog.cpp b/Engine/Source/Editor/GameProjectGeneration/Private/SProjectDialog.cpp
--- a/Engine/Source/Editor/GameProjectGeneration/Private/SProjectDialog.cpp
+++ b/Engine/Source/Editor/GameProjectGeneration/Private/SProjectDialog.cpp
@@ -496,6 +496,13 @@ void SProjectDialog::Construct(const FArguments& InArgs, EProjectDialogModeMode
 	PopulateTemplateCategories();
 
 	ProjectBrowser = SNew(SProjectBrowser);
+	if (UEditorSettings* EditorSettings = GetMutableDefault<UEditorSettings>())
+	{
+		PRAGMA_DISABLE_DEPRECATION_WARNINGS
+		bLoadMostRecentlyLoadedProject = EditorSettings->bLoadTheMostRecentlyLoadedProjectAtStartup;
+		PRAGMA_ENABLE_DEPRECATION_WARNINGS
+		EditorSettings->OnMostRecentProjectSettingChanged().AddSP(this, &SProjectDialog::OnLoadMostRecentProjectEditorSettingChanged);
+	}
 
 	ChildSlot
 	[
@@ -563,6 +570,12 @@ void SProjectDialog::Construct(const FArguments& InArgs, EProjectDialogModeMode
 
 SProjectDialog::~SProjectDialog()
 {
+	// Unbind the callback on the setting changed
+	if (UEditorSettings* EditorSettings = GetMutableDefault<UEditorSettings>())
+	{
+		EditorSettings->OnMostRecentProjectSettingChanged().RemoveAll(this);
+	}
+
 	// remove any UTemplateProjectDefs we were keeping alive
 	for (const TPair<FName, TArray<TSharedPtr<FTemplateItem>>>& Pair : Templates)
 	{
@@ -1323,7 +1336,7 @@ TSharedRef<SWidget> SProjectDialog::MakeOpenProjectPathArea()
 			+SHorizontalBox::Slot()
 			[
 				SNew(SCheckBox)
-				.IsChecked(GetDefault<UEditorSettings>()->bLoadTheMostRecentlyLoadedProjectAtStartup ? ECheckBoxState::Checked : ECheckBoxState::Unchecked)
+				.IsChecked(this, &SProjectDialog::IsLoadMostRecentProjectChecked)
 				.OnCheckStateChanged(ProjectBrowser.ToSharedRef(), &SProjectBrowser::OnAutoloadLastProjectChanged)
 				.Padding(FMargin(4.0f, 0.0f))
 				[
@@ -1814,6 +1827,19 @@ FReply SProjectDialog::OnRecentProjectsClicked()
 	return FReply::Handled();
 }
 
+void SProjectDialog::OnLoadMostRecentProjectEditorSettingChanged(bool bInAutoLoadOption)
+{
+	if (bLoadMostRecentlyLoadedProject != bInAutoLoadOption)
+	{
+		bLoadMostRecentlyLoadedProject = bInAutoLoadOption;
+	}
+}
+
+ECheckBoxState SProjectDialog::IsLoadMostRecentProjectChecked() const
+{
+	return bLoadMostRecentlyLoadedProject ? ECheckBoxState::Checked : ECheckBoxState::Unchecked;
+}
+
 FProjectInformation SProjectDialog::CreateProjectInfo() const
 {
 	TSharedPtr<FTemplateItem> SelectedTemplate = GetSelectedTemplateItem();
diff --git a/Engine/Source/Editor/GameProjectGeneration/Private/SProjectDialog.h b/Engine/Source/Editor/GameProjectGeneration/Private/SProjectDialog.h
--- a/Engine/Source/Editor/GameProjectGeneration/Private/SProjectDialog.h
+++ b/Engine/Source/Editor/GameProjectGeneration/Private/SProjectDialog.h
@@ -206,6 +206,12 @@ private:
 
 	FReply OnRecentProjectsClicked();
 
+	/** Called from the editor setting when bLoadTheMostRecentlyLoadedProjectAtStartup change to sync it */
+	void OnLoadMostRecentProjectEditorSettingChanged(bool bInAutoLoadOption);
+
+	/** Checked if it should load the most recent project, unchecked otherwise */
+	ECheckBoxState IsLoadMostRecentProjectChecked() const;
+
 	FProjectInformation CreateProjectInfo() const;
 	bool CreateProject(const FString& ProjectFile);
 	void CreateAndOpenProject();
@@ -284,4 +290,6 @@ private:
 	/** True if the last NameAndLocation validity check returned that the project path is valid for creation */
 	bool bLastNameAndLocationValidityCheckSuccessful;
 
+	/** Synced setting to the EditorSetting bLoadTheMostRecentlyLoadedProjectAtStartup */
+	bool bLoadMostRecentlyLoadedProject = false;
 };
diff --git a/Engine/Source/Editor/MainFrame/Private/HomeScreen/HomeScreenCommon.h b/Engine/Source/Editor/MainFrame/Private/HomeScreen/HomeScreenCommon.h
deleted file mode 100644
diff --git a/Engine/Source/Editor/MainFrame/Private/HomeScreen/HomeScreenSettings.cpp b/Engine/Source/Editor/MainFrame/Private/HomeScreen/HomeScreenSettings.cpp
new file mode 100644
--- /dev/null
+++ b/Engine/Source/Editor/MainFrame/Private/HomeScreen/HomeScreenSettings.cpp
@@ -0,0 +1,48 @@
+// Copyright Epic Games, Inc. All Rights Reserved.
+
+#include "Settings/HomeScreenSettings.h"
+#include "Settings/HomeScreenCommon.h"
+#include "Settings/EditorSettings.h"
+
+void UHomeScreenSettings::PostEditChangeProperty(FPropertyChangedEvent& PropertyChangedEvent)
+{
+	Super::PostEditChangeProperty(PropertyChangedEvent);
+	FName NAME_LoadAtStartup = GET_MEMBER_NAME_CHECKED(UHomeScreenSettings, LoadAtStartup);
+
+	if (PropertyChangedEvent.Property)
+	{
+		if (PropertyChangedEvent.Property->GetFName() == NAME_LoadAtStartup)
+		{
+			UEditorSettings *Settings = GetMutableDefault<UEditorSettings>();
+
+			PRAGMA_DISABLE_DEPRECATION_WARNINGS
+			Settings->bLoadTheMostRecentlyLoadedProjectAtStartup = LoadAtStartup == EAutoLoadProject::LastProject;
+			PRAGMA_ENABLE_DEPRECATION_WARNINGS
+
+			FProperty* AutoloadProjectProperty = FindFProperty<FProperty>(Settings->GetClass(), "bLoadTheMostRecentlyLoadedProjectAtStartup");
+			if (AutoloadProjectProperty != nullptr)
+			{
+				FPropertyChangedEvent PropertyUpdateStruct(AutoloadProjectProperty);
+				Settings->PostEditChangeProperty(PropertyUpdateStruct);
+			}
+			// Do not broadcast it here because the EditorSettings is going to call our OnMostRecentProjectEditorSettingChanged which will broadcast the change from the HomeScreenSettings
+		}
+	}
+}
+
+void UHomeScreenSettings::PostInitProperties()
+{
+	Super::PostInitProperties();
+	if (UEditorSettings *Settings = GetMutableDefault<UEditorSettings>())
+	{
+		Settings->OnMostRecentProjectSettingChanged().AddUObject(this, &UHomeScreenSettings::OnMostRecentProjectEditorSettingChanged);
+	}
+}
+
+void UHomeScreenSettings::OnMostRecentProjectEditorSettingChanged(bool bInAutoLoad)
+{
+	LoadAtStartup = bInAutoLoad ? EAutoLoadProject::LastProject : EAutoLoadProject::HomeScreen;
+
+	// Broadcast that the setting changed
+	OnLoadAtStartupChanged().Broadcast(LoadAtStartup);
+}
diff --git a/Engine/Source/Editor/MainFrame/Private/HomeScreen/HomeScreenWeb.h b/Engine/Source/Editor/MainFrame/Private/HomeScreen/HomeScreenWeb.h
--- a/Engine/Source/Editor/MainFrame/Private/HomeScreen/HomeScreenWeb.h
+++ b/Engine/Source/Editor/MainFrame/Private/HomeScreen/HomeScreenWeb.h
@@ -1,7 +1,7 @@
 ﻿// Copyright Epic Games, Inc. All Rights Reserved.
 
 #pragma once
-#include "HomeScreenCommon.h"
+#include "Settings/HomeScreenCommon.h"
 #include "UObject/Object.h"
 #include "HomeScreenWeb.generated.h"
 
diff --git a/Engine/Source/Editor/MainFrame/Private/HomeScreen/SHomeScreen.cpp b/Engine/Source/Editor/MainFrame/Private/HomeScreen/SHomeScreen.cpp
--- a/Engine/Source/Editor/MainFrame/Private/HomeScreen/SHomeScreen.cpp
+++ b/Engine/Source/Editor/MainFrame/Private/HomeScreen/SHomeScreen.cpp
@@ -10,6 +10,7 @@
 #include "GameProjectGenerationModule.h"
 #include "GameProjectUtils.h"
 #include "HomeScreenMenuContext.h"
+#include "Settings/HomeScreenSettings.h"
 #include "HomeScreenWeb.h"
 #include "HttpModule.h"
 #include "HttpRetrySystem.h"
@@ -79,7 +80,10 @@ SHomeScreen::SHomeScreen()
 
 SHomeScreen::~SHomeScreen()
 {
-	SaveSettings();
+	if (UHomeScreenSettings* HomeScreenSettings = GetMutableDefault<UHomeScreenSettings>())
+	{
+		HomeScreenSettings->OnLoadAtStartupChanged().RemoveAll(this);
+	}
 
 	if (GEditor && GEditor->IsTimerManagerValid())
 	{
@@ -90,24 +94,40 @@ SHomeScreen::~SHomeScreen()
 	{
 		WebObject->OnNavigationChanged().RemoveAll(this);
 	}
+
+	if (HttpRetryRequest.IsValid())
+	{
+		HttpRetryRequest->OnProcessRequestComplete().Unbind();
+		HttpRetryRequest.Reset();
+	}
+
+	if (HttpRetryManager.IsValid())
+	{
+		HttpRetryManager.Reset();
+	}
 }
 
 void SHomeScreen::Construct(const FArguments& InArgs, const TSharedPtr<SWebBrowser> InWebBrowser, const TWeakObjectPtr<UHomeScreenWeb> InWebObject)
 {
 	WebBrowser = InWebBrowser;
 	WebObject = InWebObject;
+
+	AutoLoadProjectComboBoxSelection = EAutoLoadProject::HomeScreen;
+	if (UHomeScreenSettings* HomeScreenSettings = GetMutableDefault<UHomeScreenSettings>())
+	{
+		HomeScreenSettings->OnLoadAtStartupChanged().AddSP(this, &SHomeScreen::OnLoadAtStartupSettingChanged);
+		AutoLoadProjectComboBoxSelection = HomeScreenSettings->LoadAtStartup;
+	}
+
 	if (WebObject.IsValid())
 	{
 		WebObject->OnNavigationChanged().AddSP(this, &SHomeScreen::OnNavigateToSection);
 		WebObject->OnTutorialProjectRequested().AddSP(this, &SHomeScreen::OnOpenGettingStartedProject);
 	}
 
-	AutoLoadProjectComboBoxSelection = EAutoLoadProject::HomeScreen;
 	bHasLatestEngineProjects = HasAlreadyLatestEngineProject();
 	MainHomeSelection = bHasLatestEngineProjects ? EMainSectionMenu::Home : EMainSectionMenu::GettingStarted;
 
-	LoadSettings();
-
 	// Fixed Button Size
 	constexpr float SectionButtonCheckBoxHeight = 28.f;
 
@@ -675,21 +695,9 @@ void SHomeScreen::OnMainHomeSectionChanged(ECheckBoxState InCheckBoxState, EMain
 	}
 }
 
-void SHomeScreen::LoadSettings()
-{
-	int32 LoadingType;
-	if (GConfig->GetInt(TEXT("HomeScreen"), TEXT("Loading"), LoadingType, GEngineIni))
-	{
-		// Clamp value to normal range and update loading type
-		LoadingType = FMath::Clamp<int32>(LoadingType, 0, (int32)EAutoLoadProject::MAX-1);
-
-		AutoLoadProjectComboBoxSelection = (EAutoLoadProject)LoadingType;
-	}
-}
-
-void SHomeScreen::SaveSettings()
+void SHomeScreen::OnLoadAtStartupSettingChanged(EAutoLoadProject InAutoLoadOption)
 {
-	GConfig->SetInt(TEXT("HomeScreen"), TEXT("Loading"), (int32)AutoLoadProjectComboBoxSelection, GEngineIni);
+	AutoLoadProjectComboBoxSelection = InAutoLoadOption;
 }
 
 FReply SHomeScreen::OnAutoLoadOptionChanged(EAutoLoadProject InAutoLoadOption)
@@ -699,17 +707,18 @@ FReply SHomeScreen::OnAutoLoadOptionChanged(EAutoLoadProject InAutoLoadOption)
 		return FReply::Handled();
 	}
 
-	UEditorSettings *Settings = GetMutableDefault<UEditorSettings>();
-	AutoLoadProjectComboBoxSelection = InAutoLoadOption;
-
-	Settings->bLoadTheMostRecentlyLoadedProjectAtStartup = AutoLoadProjectComboBoxSelection == EAutoLoadProject::LastProject;
-
-	// Manually trigger the property changed event so that the editor settings UI updates to reflect the change
-	FProperty* AutoloadProjectProperty = FindFProperty<FProperty>(Settings->GetClass(), "bLoadTheMostRecentlyLoadedProjectAtStartup");
-	if (AutoloadProjectProperty != NULL)
+	// Do not change the option if we can't change the setting otherwise they are not synced, the PostEditChangeProperty will handle the call to set our ComboBoxSelection.
+	if (UHomeScreenSettings* HomeScreenSettings = GetMutableDefault<UHomeScreenSettings>())
 	{
-		FPropertyChangedEvent PropertyUpdateStruct(AutoloadProjectProperty);
-		Settings->PostEditChangeProperty(PropertyUpdateStruct);
+		HomeScreenSettings->LoadAtStartup = InAutoLoadOption;
+
+		// Manually trigger the property changed event so that the home screen editor settings UI updates to reflect the change
+		FProperty* LoadAtStartupProperty = FindFProperty<FProperty>(HomeScreenSettings->GetClass(), "LoadAtStartup");
+		if (LoadAtStartupProperty != NULL)
+		{
+			FPropertyChangedEvent PropertyUpdateStruct(LoadAtStartupProperty);
+			HomeScreenSettings->PostEditChangeProperty(PropertyUpdateStruct);
+		}
 	}
 
 	if (TSharedPtr<SWidget> ComboButtonMenu = ComboButtonMenuWeak.Pin())
@@ -717,7 +726,6 @@ FReply SHomeScreen::OnAutoLoadOptionChanged(EAutoLoadProject InAutoLoadOption)
 		FSlateApplication::Get().DismissMenuByWidget(ComboButtonMenu.ToSharedRef());
 	}
 
-	SaveSettings();
 	return FReply::Handled();
 }
 
@@ -993,7 +1001,7 @@ void SHomeScreen::OnOpenGettingStartedProject()
 
 FText SHomeScreen::GetAutoLoadProjectComboBoxLabelText() const
 {
-	return AutoLoadProjectComboBoxSelection == EAutoLoadProject::HomeScreen ? LOCTEXT("HomeScreenComboHomeTab", "Home Tab") : LOCTEXT("HomeScreenComboMostRecentProject", "Most Recent Project");
+	return AutoLoadProjectComboBoxSelection == EAutoLoadProject::HomeScreen ? LOCTEXT("HomeScreenComboHomePanel", "Home Panel") : LOCTEXT("HomeScreenComboMostRecentProject", "Most Recent Project");
 }
 
 FReply SHomeScreen::OnInternetConnectionRetried()
@@ -1107,7 +1115,7 @@ FDelayedAutoRegisterHelper SHomeScreen::LoadStartupComboButtonRegistration(
 										.AutoWidth()
 										[
 											SNew(STextBlock)
-											.Text(LOCTEXT("HomeScreenHomeTab", "Home Tab"))
+											.Text(LOCTEXT("HomeScreenHomePanel", "Home Panel"))
 											.Font(FAppStyle::GetFontStyle("HomeScreen.Font.DefaultSize"))
 											.ColorAndOpacity(FStyleColors::White)
 										]
diff --git a/Engine/Source/Editor/MainFrame/Private/HomeScreen/SHomeScreen.h b/Engine/Source/Editor/MainFrame/Private/HomeScreen/SHomeScreen.h
--- a/Engine/Source/Editor/MainFrame/Private/HomeScreen/SHomeScreen.h
+++ b/Engine/Source/Editor/MainFrame/Private/HomeScreen/SHomeScreen.h
@@ -2,7 +2,7 @@
 
 #pragma once
 #include "Engine/TimerHandle.h"
-#include "HomeScreenCommon.h"
+#include "Settings/HomeScreenCommon.h"
 #include "HomeScreenWeb.h"
 #include "HttpRetrySystem.h"
 #include "Input/Reply.h"
@@ -16,13 +16,6 @@ class SComboButton;
 class SWebBrowser;
 template <typename ItemType> class STileView;
 
-enum class EAutoLoadProject
-{
-	HomeScreen,
-	LastProject,
-	MAX
-};
-
 class SHomeScreen : public SCompoundWidget
 {
 public:
@@ -50,11 +43,8 @@ private:
 	/** Handles changes to the main home section */
 	void OnMainHomeSectionChanged(ECheckBoxState InCheckBoxState, EMainSectionMenu InMainHomeSelection);
 
-	/** Loads settings from config */
-	void LoadSettings();
-
-	/** Saves settings to config */
-	void SaveSettings();
+	/** Callback from the HomeScreenSettings when the LoadAtStartup property change */
+	void OnLoadAtStartupSettingChanged(EAutoLoadProject InAutoLoadOption);
 
 	/** Handles changes to autoload last project selection */
 	FReply OnAutoLoadOptionChanged(EAutoLoadProject InAutoLoadOption);
diff --git a/Engine/Source/Editor/MainFrame/Public/Settings/HomeScreenCommon.h b/Engine/Source/Editor/MainFrame/Public/Settings/HomeScreenCommon.h
new file mode 100644
--- /dev/null
+++ b/Engine/Source/Editor/MainFrame/Public/Settings/HomeScreenCommon.h
@@ -0,0 +1,22 @@
+﻿// Copyright Epic Games, Inc. All Rights Reserved.
+
+#pragma once
+#include "HomeScreenCommon.generated.h"
+
+UENUM()
+enum class EMainSectionMenu : uint8
+{
+	None,
+	Home,
+	News,
+	GettingStarted,
+	SampleProjects
+};
+
+UENUM()
+enum class EAutoLoadProject
+{
+	HomeScreen UMETA(DisplayName = "Home Panel"),
+	LastProject UMETA(DisplayName = "Most Recent Project"),
+	MAX UMETA(Hidden)
+};
diff --git a/Engine/Source/Editor/MainFrame/Public/Settings/HomeScreenSettings.h b/Engine/Source/Editor/MainFrame/Public/Settings/HomeScreenSettings.h
new file mode 100644
--- /dev/null
+++ b/Engine/Source/Editor/MainFrame/Public/Settings/HomeScreenSettings.h
@@ -0,0 +1,32 @@
+// Copyright Epic Games, Inc. All Rights Reserved.
+
+#pragma once
+#include "HomeScreenCommon.h"
+#include "UObject/Object.h"
+#include "HomeScreenSettings.generated.h"
+
+DECLARE_MULTICAST_DELEGATE_OneParam(FOnLoadAtStartupChanged, EAutoLoadProject)
+
+// Load on Startup setting is registered with the FEditorLoadingSavingSettingsCustomization in the Loading & Saving Section
+
+UCLASS(MinimalAPI, config=EditorSettings)
+class UHomeScreenSettings : public UObject
+{
+	GENERATED_BODY()
+
+public:
+	UPROPERTY()
+	EAutoLoadProject LoadAtStartup = EAutoLoadProject::HomeScreen;
+
+public:
+	MAINFRAME_API virtual void PostEditChangeProperty(FPropertyChangedEvent& PropertyChangedEvent) override;
+	MAINFRAME_API virtual void PostInitProperties() override;
+
+	FOnLoadAtStartupChanged& OnLoadAtStartupChanged() { return OnLoadAtStartupChangedDelegate; }
+
+private:
+	void OnMostRecentProjectEditorSettingChanged(bool bInAutoLoad);
+
+private:
+	FOnLoadAtStartupChanged OnLoadAtStartupChangedDelegate;
+};
diff --git a/Engine/Source/Editor/UnrealEd/Private/EditorEngine.cpp b/Engine/Source/Editor/UnrealEd/Private/EditorEngine.cpp
--- a/Engine/Source/Editor/UnrealEd/Private/EditorEngine.cpp
+++ b/Engine/Source/Editor/UnrealEd/Private/EditorEngine.cpp
@@ -280,6 +280,7 @@
 #include "TargetReceipt.h"
 
 #include "AudioEditorSettings.h"
+#include "Settings/HomeScreenCommon.h"
 
 #include UE_INLINE_GENERATED_CPP_BY_NAME(EditorEngine)
 
@@ -5628,7 +5629,22 @@ void UEditorEngine::UpdateAutoLoadProject()
 		bShouldLoadRecentProjects = false;
 	}
 
+	PRAGMA_DISABLE_DEPRECATION_WARNINGS
 	GetMutableDefault<UEditorSettings>()->bLoadTheMostRecentlyLoadedProjectAtStartup = bShouldLoadRecentProjects;
+	PRAGMA_ENABLE_DEPRECATION_WARNINGS
+
+	// Sync up the new LoadAtStartup setting with the old one
+	if (UClass* SettingsClass = FindObject<UClass>(nullptr, TEXT("/Script/MainFrame.HomeScreenSettings")))
+	{
+		UObject* DefaultObj = SettingsClass->GetDefaultObject();
+		FProperty* Prop = SettingsClass->FindPropertyByName("LoadAtStartup");
+		if (Prop && DefaultObj && Prop->IsA<FEnumProperty>())
+		{
+			FEnumProperty* AutoLoadProperty = CastFieldChecked<FEnumProperty>(Prop);
+			EAutoLoadProject Value = bShouldLoadRecentProjects ? EAutoLoadProject::LastProject : EAutoLoadProject::HomeScreen;
+			AutoLoadProperty->SetValue_InContainer(DefaultObj, &Value);
+		}
+	}
 
 #if PLATFORM_MAC
 	if ( !GIsBuildMachine )
diff --git a/Engine/Source/Editor/UnrealEd/Private/Settings/EditorLoadingSavingSettingsCustomization.h b/Engine/Source/Editor/UnrealEd/Private/Settings/EditorLoadingSavingSettingsCustomization.h
--- a/Engine/Source/Editor/UnrealEd/Private/Settings/EditorLoadingSavingSettingsCustomization.h
+++ b/Engine/Source/Editor/UnrealEd/Private/Settings/EditorLoadingSavingSettingsCustomization.h
@@ -48,7 +48,7 @@ protected:
 	/**
 	 * Customizes the Startup property category.
 	 *
-	 * This customization pulls in a setting from the game agnostic Editor settings, which
+	 * This customization pulls in a setting from the game agnostic HomeScreenSettings in the MainFrame module, which
 	 * are stored in a different UObject, but which we would like to show in this section.
 	 *
 	 * @param LayoutBuilder The layout builder.
@@ -57,9 +57,12 @@ protected:
 	{
 		IDetailCategoryBuilder& StartupCategory = LayoutBuilder.EditCategory("Startup");
 		{
-			TArray<UObject*> ObjectList;
-			ObjectList.Add(GetMutableDefault<UEditorSettings>());
-			StartupCategory.AddExternalObjectProperty(ObjectList, "bLoadTheMostRecentlyLoadedProjectAtStartup");
+			if (const UClass* SettingsClass = FindObject<UClass>(nullptr, TEXT("/Script/MainFrame.HomeScreenSettings")))
+			{
+				TArray<UObject*> HomeScreenSettingObjectList;
+				HomeScreenSettingObjectList.Add(SettingsClass->GetDefaultObject());
+				StartupCategory.AddExternalObjectProperty(HomeScreenSettingObjectList, "LoadAtStartup");
+			}
 		}
 	}
 };
diff --git a/Engine/Source/Editor/UnrealEd/Private/Settings/EditorSettings.cpp b/Engine/Source/Editor/UnrealEd/Private/Settings/EditorSettings.cpp
--- a/Engine/Source/Editor/UnrealEd/Private/Settings/EditorSettings.cpp
+++ b/Engine/Source/Editor/UnrealEd/Private/Settings/EditorSettings.cpp
@@ -64,7 +64,9 @@ void UEditorSettings::PostEditChangeProperty( struct FPropertyChangedEvent& Prop
 	FProperty* PropertyThatChanged = PropertyChangedEvent.MemberProperty;
 	const FName PropertyName = PropertyThatChanged ? PropertyThatChanged->GetFName() : NAME_None;
 
-	if (PropertyName == GET_MEMBER_NAME_CHECKED(UEditorSettings, bLoadTheMostRecentlyLoadedProjectAtStartup))
+	PRAGMA_DISABLE_DEPRECATION_WARNINGS
+	FName DeprecatedName = GET_MEMBER_NAME_CHECKED(UEditorSettings, bLoadTheMostRecentlyLoadedProjectAtStartup);
+	if (PropertyName == DeprecatedName)
 	{
 		const FString& AutoLoadProjectFileName = IProjectManager::Get().GetAutoLoadProjectFileName();
 		if ( bLoadTheMostRecentlyLoadedProjectAtStartup )
@@ -77,7 +79,11 @@ void UEditorSettings::PostEditChangeProperty( struct FPropertyChangedEvent& Prop
 			// Remove the file. It's possible for bLoadTheMostRecentlyLoadedProjectAtStartup to be set before FPaths::GetProjectFilePath() is valid, so we need to distinguish the two cases.
 			IFileManager::Get().Delete(*AutoLoadProjectFileName);
 		}
+
+		// Broadcast that this setting changed
+		OnMostRecentProjectSettingChanged().Broadcast(bLoadTheMostRecentlyLoadedProjectAtStartup);
 	}
+	PRAGMA_ENABLE_DEPRECATION_WARNINGS
 	else if (PropertyName == GET_MEMBER_NAME_CHECKED(UEditorSettings, GlobalLocalDDCPath))
 	{
 		if (GlobalLocalDDCPath.Path.IsEmpty())
diff --git a/Engine/Source/Editor/UnrealEd/Public/Settings/EditorSettings.h b/Engine/Source/Editor/UnrealEd/Public/Settings/EditorSettings.h
--- a/Engine/Source/Editor/UnrealEd/Public/Settings/EditorSettings.h
+++ b/Engine/Source/Editor/UnrealEd/Public/Settings/EditorSettings.h
@@ -51,6 +51,7 @@ struct FRecentProjectFile
 	}
 };
 
+DECLARE_MULTICAST_DELEGATE_OneParam(FOnMostRecentProjectSettingChanged, bool /** bInShouldUseMostRecentProject */)
 
 UCLASS(config=EditorSettings, MinimalAPI)
 class UEditorSettings : public UObject
@@ -124,9 +125,10 @@ class UEditorSettings : public UObject
 
 	// =====================================================================
 
-	/** When checked, the most recently loaded project will be auto-loaded at editor startup if no other project was specified on the command line */
-	UPROPERTY()
+	UE_DEPRECATED(5.7, "Use Load At Startup setting")
+	UPROPERTY(meta = (DeprecatedProperty, DeprecatedMessage="Use Load At Startup setting"))
 	bool bLoadTheMostRecentlyLoadedProjectAtStartup; // Note that this property is NOT config since it is not necessary to save the value to ini. It is determined at startup in UEditorEngine::InitEditor().
+	FOnMostRecentProjectSettingChanged& OnMostRecentProjectSettingChanged() { return OnMostRecentProjectSettingChangedDelegate; }
 
 	// =====================================================================
 	// The following options are NOT exposed in the preferences Editor
@@ -173,4 +175,7 @@ class UEditorSettings : public UObject
 	UNREALED_API virtual bool CanEditChange(const FProperty* InProperty) const override;
 	UNREALED_API virtual void PostEditChangeProperty( struct FPropertyChangedEvent& PropertyChangedEvent) override;
 	//~ End UObject Interface
+
+private:
+	FOnMostRecentProjectSettingChanged OnMostRecentProjectSettingChangedDelegate;
 };
