[HomeScreen] Unbind callback to WebObject when HomeScreen is destructed

diff --git a/Engine/Source/Editor/MainFrame/Private/HomeScreen/SHomeScreen.cpp b/Engine/Source/Editor/MainFrame/Private/HomeScreen/SHomeScreen.cpp
--- a/Engine/Source/Editor/MainFrame/Private/HomeScreen/SHomeScreen.cpp
+++ b/Engine/Source/Editor/MainFrame/Private/HomeScreen/SHomeScreen.cpp
@@ -93,6 +93,7 @@ SHomeScreen::~SHomeScreen()
 	if (WebObject.IsValid())
 	{
 		WebObject->OnNavigationChanged().RemoveAll(this);
+		WebObject->OnTutorialProjectRequested().RemoveAll(this);
 	}
 
 	if (HttpRetryRequest.IsValid())
@@ -101,10 +102,7 @@ SHomeScreen::~SHomeScreen()
 		HttpRetryRequest.Reset();
 	}
 
-	if (HttpRetryManager.IsValid())
-	{
-		HttpRetryManager.Reset();
-	}
+	HttpRetryManager.Reset();
 }
 
 void SHomeScreen::Construct(const FArguments& InArgs, const TSharedPtr<SWebBrowser> InWebBrowser, const TWeakObjectPtr<UHomeScreenWeb> InWebObject)
