[HomeScreen] Fix ProjectBrowser not closing when creating or selecting a project

Create the project window in the HomeScreen to not be Modal so that it can be closed when creating or selecting a project.

diff --git a/Engine/Source/Editor/MainFrame/Private/HomeScreen/SHomeScreen.cpp b/Engine/Source/Editor/MainFrame/Private/HomeScreen/SHomeScreen.cpp
--- a/Engine/Source/Editor/MainFrame/Private/HomeScreen/SHomeScreen.cpp
+++ b/Engine/Source/Editor/MainFrame/Private/HomeScreen/SHomeScreen.cpp
@@ -646,8 +646,17 @@ FReply SHomeScreen::OnCreateProjectDialog(bool bInAllowProjectOpening, bool bInA
 		FMainFrameActionCallbacks::CreateProjectBrowserWindow(bInAllowProjectOpening, bInAllowProjectCreation);
 
 	NewProjectWindow->SetContent(FGameProjectGenerationModule::Get().CreateGameProjectDialog(bInAllowProjectOpening, bInAllowProjectCreation));
-	FSlateApplication::Get().AddModalWindow(NewProjectWindow, nullptr);
-	
+
+	IMainFrameModule& MainFrameModule = FModuleManager::GetModuleChecked<IMainFrameModule>(TEXT("MainFrame"));
+	if (MainFrameModule.GetParentWindow().IsValid())
+	{
+		FSlateApplication::Get().AddWindowAsNativeChild(NewProjectWindow, MainFrameModule.GetParentWindow().ToSharedRef());
+	}
+	else
+	{
+		FSlateApplication::Get().AddWindow(NewProjectWindow);
+	}
+
 	return FReply::Handled();
 }
 
