[HomeScreen] Change URL to production ones which will be avaiable in the final release

Added also WebPage localization to match Editor Language setting

diff --git a/Engine/Source/Editor/MainFrame/Private/HomeScreen/SHomeScreen.cpp b/Engine/Source/Editor/MainFrame/Private/HomeScreen/SHomeScreen.cpp
--- a/Engine/Source/Editor/MainFrame/Private/HomeScreen/SHomeScreen.cpp
+++ b/Engine/Source/Editor/MainFrame/Private/HomeScreen/SHomeScreen.cpp
@@ -24,6 +24,7 @@
 #include "TimerManager.h"
 #include "ToolMenus.h"
 #include "UnrealEdGlobals.h"
+#include "Internationalization/Culture.h"
 #include "Widgets/Images/SImage.h"
 #include "Widgets/Images/SThrobber.h"
 #include "Widgets/Input/SButton.h"
@@ -674,25 +675,28 @@ ECheckBoxState SHomeScreen::IsMainHomeSectionChecked(EMainSectionMenu InMainHome
 void SHomeScreen::OnMainHomeSectionChanged(ECheckBoxState InCheckBoxState, EMainSectionMenu InMainHomeSelection)
 {
 	FString URL = TEXT("");
+	const FString BaseURL = TEXT("https://editor.unrealengine.com");
+	// Use the Editor Language for the WebPage language.
+	const FString FullLocale = FInternationalization::Get().GetCurrentLanguage()->GetName();
 	switch (InMainHomeSelection)
 	{
 	case EMainSectionMenu::News:
 		MainHomeSelection = EMainSectionMenu::News;
-		URL = TEXT("https://editor-website.ol.epicgames.net/en-US/news");
+		URL = BaseURL / FullLocale / TEXT("news");
 		break;
 	case EMainSectionMenu::GettingStarted:
 		MainHomeSelection = EMainSectionMenu::GettingStarted;
-		URL = TEXT("https://editor-website.ol.epicgames.net/en-US/get-started-with-unreal-engine");
+		URL = BaseURL / FullLocale / TEXT("get-started-with-unreal-engine");
 		break;
 	case EMainSectionMenu::SampleProjects:
 		MainHomeSelection = EMainSectionMenu::SampleProjects;
-		URL = TEXT("https://editor-website.ol.epicgames.net/en-US/samples");
+		URL = BaseURL / FullLocale / TEXT("samples");
 		break;
 	case EMainSectionMenu::None:
 	case EMainSectionMenu::Home:
 	default:
 		MainHomeSelection = EMainSectionMenu::Home;
-		URL = bHasLatestEngineProjects ? TEXT("https://editor-website.ol.epicgames.net/en-US/recent-projects") : TEXT("https://editor-website.ol.epicgames.net/en-US/no-projects");
+		URL = bHasLatestEngineProjects ? BaseURL / FullLocale / TEXT("recent-projects") : BaseURL / FullLocale / TEXT("no-projects");
 		break;
 	}
 
diff --git a/Engine/Source/Editor/MainFrame/Private/MainFrameModule.cpp b/Engine/Source/Editor/MainFrame/Private/MainFrameModule.cpp
--- a/Engine/Source/Editor/MainFrame/Private/MainFrameModule.cpp
+++ b/Engine/Source/Editor/MainFrame/Private/MainFrameModule.cpp
@@ -958,7 +958,7 @@ TSharedRef<SWidget> FMainFrameModule::CreateHomeScreenWidget()
 		IWebBrowserSingleton* WebBrowserSingleton = WebBrowserModule.GetSingleton();
 
 		FCreateBrowserWindowSettings WindowSettings;
-		WindowSettings.InitialURL = TEXT("https://editor-website.ol.epicgames.net/en-US");
+		WindowSettings.InitialURL = TEXT("https://editor.unrealengine.com/en-US/recent-projects");
 		WindowSettings.bUseTransparency = true;
 		WindowSettings.BrowserFrameRate = 60;
 
