Fix invalid FICUTextBiDi memory read race condition crash on editor shutdown.

On behalf of Jaime.Dale.

#tests editor 3x

diff --git a/Engine/Source/Editor/MainFrame/Private/MainFrameModule.cpp b/Engine/Source/Editor/MainFrame/Private/MainFrameModule.cpp
--- a/Engine/Source/Editor/MainFrame/Private/MainFrameModule.cpp
+++ b/Engine/Source/Editor/MainFrame/Private/MainFrameModule.cpp
@@ -712,6 +712,7 @@ void FMainFrameModule::ShutdownModule( )
 	WebObject.Reset();
 	WebBrowserWindow.Reset();
 	WebBrowser.Reset();
+	HomeScreenTabManager.Reset();
 
 	// Unregister HomeScreen tab on shutdown
 	FGlobalTabmanager::Get()->UnregisterTabSpawner(TEXT("HomeScreen"));
