Implement mac tab workaround for the web browser issues on that platform for the homepage.

diff --git a/Engine/Source/Editor/MainFrame/Private/MainFrameModule.cpp b/Engine/Source/Editor/MainFrame/Private/MainFrameModule.cpp
--- a/Engine/Source/Editor/MainFrame/Private/MainFrameModule.cpp
+++ b/Engine/Source/Editor/MainFrame/Private/MainFrameModule.cpp
@@ -150,7 +150,7 @@ TSharedRef<SWidget> FProjectDialogProvider::CreateMainFrameContentWidget() const
 {
 	if (UE::Editor::HomeScreen::IsHomeScreenEnabled())
 	{
-		return FMainFrameModule::CreateHomeScreenWidget();
+		return FMainFrameModule::CreateHomeScreenWidget({});
 	}
 
 	return FGameProjectGenerationModule::Get().CreateGameProjectDialog(/*bAllowProjectOpening=*/true, /*bAllowProjectCreate=*/true);
@@ -935,12 +935,12 @@ TSharedRef<SDockTab> FMainFrameModule::SpawnHomeScreenWindow(const FSpawnTabArgs
 	HomeScreenTab->SetTabIcon(FAppStyle::Get().GetBrush("HomeScreen.Home"));
 
 	// This will create the necessary object and widgets needed for the HomeScreen and return a newly created one or NullWidget if something went wrong
-	HomeScreenTab->SetContent(FMainFrameModule::CreateHomeScreenWidget());
+	HomeScreenTab->SetContent(FMainFrameModule::CreateHomeScreenWidget(HomeScreenTab));
 
 	return HomeScreenTab;
 }
 
-TSharedRef<SWidget> FMainFrameModule::CreateHomeScreenWidget()
+TSharedRef<SWidget> FMainFrameModule::CreateHomeScreenWidget(const TSharedPtr<SDockTab>& InParentDockTab)
 {
 	if (!WebObject.IsValid())
 	{
@@ -964,13 +964,16 @@ TSharedRef<SWidget> FMainFrameModule::CreateHomeScreenWidget()
 		WindowSettings.BrowserFrameRate = 60;
 
 		WebBrowserWindow = WebBrowserSingleton->CreateBrowserWindow(WindowSettings);
-
 		if (!WebBrowserWindow.IsValid())
 		{
 			FMessageDialog::Open(EAppMsgType::Ok, LOCTEXT("WebBrowserWindowCreationIssue", "Failed to create the web browser window."));
 			return SNullWidget::NullWidget;
 		}
 		
+		if (InParentDockTab)
+		{
+			WebBrowserWindow->SetParentDockTab(InParentDockTab);
+		}
 		WebBrowserSingleton->SetDevToolsShortcutEnabled(true);
 		WebBrowserWindow->OnCreateWindow().BindLambda([](const TWeakPtr<IWebBrowserWindow>& InNewBrowserWindowWeak, const TWeakPtr<IWebBrowserPopupFeatures>& PopupFeatures)
 		{
diff --git a/Engine/Source/Editor/MainFrame/Private/MainFrameModule.h b/Engine/Source/Editor/MainFrame/Private/MainFrameModule.h
--- a/Engine/Source/Editor/MainFrame/Private/MainFrameModule.h
+++ b/Engine/Source/Editor/MainFrame/Private/MainFrameModule.h
@@ -228,7 +228,7 @@ public:
 	static TSharedRef<SDockTab> SpawnHomeScreenWindow(const FSpawnTabArgs& InArgs);
 
 	// Create the home screen widget
-	static TSharedRef<SWidget> CreateHomeScreenWidget();
+	static TSharedRef<SWidget> CreateHomeScreenWidget(const TSharedPtr<SDockTab>& InParentTab);
 
 private:
 
