Home Panel - Roadmap links

diff --git a/Engine/Source/Editor/MainFrame/Private/HomeScreen/SHomeScreen.cpp b/Engine/Source/Editor/MainFrame/Private/HomeScreen/SHomeScreen.cpp
--- a/Engine/Source/Editor/MainFrame/Private/HomeScreen/SHomeScreen.cpp
+++ b/Engine/Source/Editor/MainFrame/Private/HomeScreen/SHomeScreen.cpp
@@ -143,7 +143,7 @@ void SHomeScreen::Construct(const FArguments& InArgs, const TSharedPtr<SWebBrows
 		LOCTEXT("ProjectHomeDocumentation", "Documentation"), FAppStyle::GetBrush("HomeScreen.Documentation"));
 	CreateResourceButtons(TutorialsButton, TEXT("https://dev.epicgames.com/community/unreal-engine/learning"),
 		LOCTEXT("ProjectHomeTutorialsButton", "Tutorials"), FAppStyle::GetBrush("HomeScreen.Tutorial"));
-	CreateResourceButtons(RoadmapButton, TEXT("https://portal.productboard.com/epicgames/1-unreal-engine-public-roadmap/tabs/109-unreal-engine-5-5"),
+	CreateResourceButtons(RoadmapButton, TEXT("https://portal.productboard.com/epicgames/1-unreal-engine-public-roadmap/tabs/127-unreal-engine-5-7"),
 		LOCTEXT("ProjectHomeRoadmapButton", "Roadmap"), FAppStyle::GetBrush("HomeScreen.Roadmap"));
 	CreateResourceButtons(ReleaseNotesButton, TEXT("https://dev.epicgames.com/documentation/en-us/unreal-engine/unreal-engine-5-6-release-notes"),
 		LOCTEXT("ProjectHomeReleaseNotesButton", "Release Notes"), FAppStyle::GetBrush("HomeScreen.ReleaseNotes"));
