[HomeScreen] Further URL security issue

diff --git a/Engine/Source/Editor/MainFrame/Private/HomeScreen/HomeScreenWeb.cpp b/Engine/Source/Editor/MainFrame/Private/HomeScreen/HomeScreenWeb.cpp
--- a/Engine/Source/Editor/MainFrame/Private/HomeScreen/HomeScreenWeb.cpp
+++ b/Engine/Source/Editor/MainFrame/Private/HomeScreen/HomeScreenWeb.cpp
@@ -15,20 +15,11 @@ void UHomeScreenWeb::NavigateTo(EMainSectionMenu InSectionToNavigate)
 
 void UHomeScreenWeb::OpenWebPage(const FString& InURL) const
 {
-	FText ErrorText = FText::Format(LOCTEXT("LaunchingURLNotAllowed", "Given URL is not registered as an allowed one:\n {0}"), FText::FromString(InURL));
-	if (!InURL.StartsWith("http"))
-	{
-		FMessageDialog::Open(EAppMsgType::Ok, ErrorText);
-		return;
-	}
-
-	FString Domain;
-	FString Protocol;
-	InURL.Split("://", &Protocol, &Domain);
+	const FText ErrorText = FText::Format(LOCTEXT("LaunchingURLNotAllowed", "Given URL is not registered as an allowed one:\n {0}"), FText::FromString(InURL));
 
-	const FString EpicOriginRegex = TEXT("^(www|www2|localhost|dev|docs|forums|company|legal)?(\\.staging)?\\.?(unrealengine|epicgames|twinmotion|metahuman|realityscan|kidswebservices|fab|quixel)(\\.com|-(ci|gamedev)\\.ol\\.epicgames\\.net)/?");
-	FRegexPattern URLPattern(EpicOriginRegex);
-	FRegexMatcher Matcher(URLPattern, Domain);
+	const FString EpicOriginRegex = TEXT("^https?://(www|www2|localhost|dev|docs|forums|company|legal)?(\\.staging)?\\.?(unrealengine|epicgames|twinmotion|metahuman|realityscan|kidswebservices|fab|quixel)(\\.com|-(ci|gamedev)\\.ol\\.epicgames\\.net)(/[a-zA-Z0-9._~:/?#@!$&'()*+,;=%-]*/?)*$");
+	const FRegexPattern URLPattern(EpicOriginRegex);
+	FRegexMatcher Matcher(URLPattern, InURL);
 
 	if (Matcher.FindNext())
 	{
