[HomeScreen] Check projects from all installed engines and not just the current one

The home screen to display the NoProject WebPage was checking if there were any latest engine version projects created with the current engine build.
Instead check all projects from all known engine builds and check that the version of it match the current engine build excluding the patch (5.7 and 5.7.1 are considered the same for example).

diff --git a/Engine/Source/Editor/MainFrame/Private/HomeScreen/SHomeScreen.cpp b/Engine/Source/Editor/MainFrame/Private/HomeScreen/SHomeScreen.cpp
--- a/Engine/Source/Editor/MainFrame/Private/HomeScreen/SHomeScreen.cpp
+++ b/Engine/Source/Editor/MainFrame/Private/HomeScreen/SHomeScreen.cpp
@@ -1023,68 +1023,67 @@ FReply SHomeScreen::OnInternetConnectionRetried()
 
 bool SHomeScreen::HasAlreadyLatestEngineProject() const
 {
-	TArray<FString> ProjectFiles;
-
-	FString RootDir;
 	IDesktopPlatform* DesktopPlatformModule = FDesktopPlatformModule::Get();
 	if (!DesktopPlatformModule)
 	{
 		return false;
 	}
 
-	const FString EngineIdentifier = DesktopPlatformModule->GetCurrentEngineIdentifier();
-
-	DesktopPlatformModule->EnumerateProjectsKnownByEngine(EngineIdentifier, false, ProjectFiles);
-	if (!DesktopPlatformModule->GetEngineRootDirFromIdentifier(EngineIdentifier, RootDir))
+	// Get all engine installation and projects known by them
+	TMap<FString, FString> EngineInstallations;
+	DesktopPlatformModule->EnumerateEngineInstallations(EngineInstallations);
+	
+	for (TMap<FString, FString>::TConstIterator Iter(EngineInstallations); Iter; ++Iter)
 	{
-		// The engine root dir couldn't be found so fake that it's a new user
-		return false;
-	}
+		TArray<FString> ProjectFiles;
 
-	ProjectFiles.RemoveAll([RootDir, DesktopPlatformModule] (const FString& InProjectFile)
+		if (DesktopPlatformModule->EnumerateProjectsKnownByEngine(Iter.Key(), false, ProjectFiles))
 		{
-			FString ProjectFilename = IFileManager::Get().ConvertToAbsolutePathForExternalAppForRead(*InProjectFile);
-			if (!FPaths::FileExists(ProjectFilename))
+			FString RootDir;
+			if (DesktopPlatformModule->GetEngineRootDirFromIdentifier(Iter.Key(), RootDir))
 			{
-				return true;
-			}
-
-			if (ProjectFilename.Contains(RootDir))
-			{
-				return true;
-			}
-			FString Identifier;
-			DesktopPlatformModule->GetEngineIdentifierForProject(ProjectFilename, Identifier);
-
-			FEngineVersion ProjectEngineVersion;
-
-			if (DesktopPlatformModule->IsStockEngineRelease(Identifier))
-			{
-				DesktopPlatformModule->TryParseStockEngineVersion(Identifier, ProjectEngineVersion);
-			}
-
-			if (ProjectEngineVersion.IsEmpty())
-			{
-				FString RootProjectDir;
-				DesktopPlatformModule->GetEngineRootDirFromIdentifier(Identifier, RootProjectDir);
-				if (!DesktopPlatformModule->TryGetEngineVersion(RootProjectDir, ProjectEngineVersion))
+				// Iterate over each found project and return true if one of them is made with this engine version
+				for (const FString& Project : ProjectFiles)
 				{
-					return true;
+					const FString ProjectFilename = IFileManager::Get().ConvertToAbsolutePathForExternalAppForRead(*Project);
+					if (!FPaths::FileExists(ProjectFilename) || ProjectFilename.Contains(RootDir))
+					{
+						continue;
+					}
+                    
+                    FString Identifier;
+                    DesktopPlatformModule->GetEngineIdentifierForProject(ProjectFilename, Identifier);
+                    
+                    FEngineVersion ProjectEngineVersion;
+                    
+                    if (DesktopPlatformModule->IsStockEngineRelease(Identifier))
+                    {
+                    	DesktopPlatformModule->TryParseStockEngineVersion(Identifier, ProjectEngineVersion);
+                    }
+                    
+                    if (ProjectEngineVersion.IsEmpty())
+                    {
+                    	FString RootProjectDir;
+                    	DesktopPlatformModule->GetEngineRootDirFromIdentifier(Identifier, RootProjectDir);
+                    	if (!DesktopPlatformModule->TryGetEngineVersion(RootProjectDir, ProjectEngineVersion))
+                    	{
+                    		continue;
+                    	}
+                    }
+                    
+                    const FString EngineVersionString = FEngineVersion::Current().ToString(EVersionComponent::Minor);
+                    const FString ProjectEngineVersionString = ProjectEngineVersion.ToString(EVersionComponent::Minor);
+                    
+                    if (EngineVersionString == ProjectEngineVersionString)
+                    {
+                    	return true;
+                    }
 				}
 			}
+		}
+	}
 
-			const FString EngineVersionString = FEngineVersion::Current().ToString(EVersionComponent::Patch);
-			const FString ProjectEngineVersionString = ProjectEngineVersion.ToString(EVersionComponent::Patch);
-
-			if (EngineVersionString != ProjectEngineVersionString)
-			{
-				return true;
-			}
-
-			return false;
-		});
-
-	return !ProjectFiles.IsEmpty();
+	return false;
 }
 
 FDelayedAutoRegisterHelper SHomeScreen::LoadStartupComboButtonRegistration(
