[HomeScreen] Change HomeScreen static variable inside MainFrame module

diff --git a/Engine/Source/Editor/MainFrame/Private/MainFrameModule.cpp b/Engine/Source/Editor/MainFrame/Private/MainFrameModule.cpp
--- a/Engine/Source/Editor/MainFrame/Private/MainFrameModule.cpp
+++ b/Engine/Source/Editor/MainFrame/Private/MainFrameModule.cpp
@@ -51,11 +51,6 @@
 DEFINE_LOG_CATEGORY(LogMainFrame);
 #define LOCTEXT_NAMESPACE "FMainFrameModule"
 
-TStrongObjectPtr<UHomeScreenWeb> FMainFrameModule::WebObject = nullptr;
-TSharedPtr<IWebBrowserWindow> FMainFrameModule::WebBrowserWindow = nullptr;
-TSharedPtr<SWebBrowser> FMainFrameModule::WebBrowser = nullptr;
-TSharedPtr<FTabManager> FMainFrameModule::HomeScreenTabManager = nullptr;
-
 namespace UE::Editor::HomeScreen
 {
 	static FAutoConsoleVariable CVarEnableHomeScreen(
@@ -100,14 +95,19 @@ const FText StaticGetApplicationTitle( const bool bIncludeGameName )
 /* FProjectDialogProvider
  *****************************************************************************/
 
-void FProjectDialogProvider::Register()
+void FProjectDialogProvider::Register(FMainFrameModule* InMainFrameModule)
 {
 	IModularFeatures::Get().RegisterModularFeature(IEditorMainFrameProvider::GetModularFeatureName(), this);
+	if (InMainFrameModule != nullptr)
+	{
+		MainFrameModulePtr = InMainFrameModule;
+	}
 }
 
 void FProjectDialogProvider::UnRegister()
 {
 	IModularFeatures::Get().UnregisterModularFeature(IEditorMainFrameProvider::GetModularFeatureName(), this);
+	MainFrameModulePtr = nullptr;
 }
 
 bool FProjectDialogProvider::IsRequestingMainFrameControl() const
@@ -150,7 +150,10 @@ TSharedRef<SWidget> FProjectDialogProvider::CreateMainFrameContentWidget() const
 {
 	if (UE::Editor::HomeScreen::IsHomeScreenEnabled())
 	{
-		return FMainFrameModule::CreateHomeScreenWidget({});
+		if (MainFrameModulePtr != nullptr)
+		{
+			return MainFrameModulePtr->CreateHomeScreenWidget({});
+		}
 	}
 
 	return FGameProjectGenerationModule::Get().CreateGameProjectDialog(/*bAllowProjectOpening=*/true, /*bAllowProjectCreate=*/true);
@@ -636,7 +639,7 @@ void FMainFrameModule::StartupModule( )
 
 	MRUFavoritesList = NULL;
 
-	ProjectDialogProvider.Register();
+	ProjectDialogProvider.Register(this);
 
 	ensureMsgf(!IsRunningGame(), TEXT("The MainFrame module should only be loaded when running the editor.  Code that extends the editor, adds menu items, etc... should not run when running in -game mode or in a non-WITH_EDITOR build"));
 	MainFrameHandler = MakeShareable(new FMainFrameHandler);
@@ -664,7 +667,7 @@ void FMainFrameModule::StartupModule( )
 	if (UE::Editor::HomeScreen::IsHomeScreenEnabled())
 	{
 		FGlobalTabmanager::Get()->RegisterTabSpawner(
-			"HomeScreen", FOnSpawnTab::CreateStatic(&FMainFrameModule::SpawnHomeScreenWindow))
+			"HomeScreen", FOnSpawnTab::CreateRaw(this, &FMainFrameModule::SpawnHomeScreenWindow))
 			.SetAutoGenerateMenuEntry(false)
 			.SetMenuType(ETabSpawnerMenuType::Hidden)
 			.SetDisplayName(FText::GetEmpty());
@@ -930,12 +933,12 @@ TSharedRef<SDockTab> FMainFrameModule::SpawnHomeScreenWindow(const FSpawnTabArgs
 	HomeScreenTabManager->SetAllowWindowMenuBar(true);
 	FToolMenuContext EmptyContext;
 	FMainFrameModule::Get().MakeMainMenu(HomeScreenTabManager, "HomeScreen.MainMenu", EmptyContext);
-	HomeScreenTab->SetOnTabClosed(SDockTab::FOnTabClosedCallback::CreateStatic([] (TSharedRef<SDockTab>) { HomeScreenTabManager.Reset(); }));
+	HomeScreenTab->SetOnTabClosed(SDockTab::FOnTabClosedCallback::CreateLambda([this] (TSharedRef<SDockTab>) { HomeScreenTabManager.Reset(); }));
 
 	HomeScreenTab->SetTabIcon(FAppStyle::Get().GetBrush("HomeScreen.Home"));
 
 	// This will create the necessary object and widgets needed for the HomeScreen and return a newly created one or NullWidget if something went wrong
-	HomeScreenTab->SetContent(FMainFrameModule::CreateHomeScreenWidget(HomeScreenTab));
+	HomeScreenTab->SetContent(CreateHomeScreenWidget(HomeScreenTab));
 
 	return HomeScreenTab;
 }
diff --git a/Engine/Source/Editor/MainFrame/Private/MainFrameModule.h b/Engine/Source/Editor/MainFrame/Private/MainFrameModule.h
--- a/Engine/Source/Editor/MainFrame/Private/MainFrameModule.h
+++ b/Engine/Source/Editor/MainFrame/Private/MainFrameModule.h
@@ -24,7 +24,7 @@ class UHomeScreenWeb;
 class FProjectDialogProvider : public IEditorMainFrameProvider
 {
 public:
-	void Register();
+	void Register(class FMainFrameModule* InMainFrameModule);
 	void UnRegister();
 
 	//~ Begin IEditorMainFrameProvider interface
@@ -32,6 +32,9 @@ public:
 	virtual FMainFrameWindowOverrides GetDesiredWindowConfiguration() const override;
 	virtual TSharedRef<SWidget> CreateMainFrameContentWidget() const override;
 	//~ End IEditorMainFrameProvider interface
+	
+private:
+	FMainFrameModule* MainFrameModulePtr = nullptr;
 };
 
 /**
@@ -224,14 +227,14 @@ public:
 	/** Get the size of the project browser window */
 	static FVector2D GetProjectBrowserWindowSize() { return FVector2D(1206, 822); }
 
-	// Spawns the home screen tab
-	static TSharedRef<SDockTab> SpawnHomeScreenWindow(const FSpawnTabArgs& InArgs);
-
 	// Create the home screen widget
-	static TSharedRef<SWidget> CreateHomeScreenWidget(const TSharedPtr<SDockTab>& InParentTab);
+	TSharedRef<SWidget> CreateHomeScreenWidget(const TSharedPtr<SDockTab>& InParentTab);
 
 private:
 
+	// Spawns the home screen tab
+	TSharedRef<SDockTab> SpawnHomeScreenWindow(const FSpawnTabArgs& InArgs);
+
 	// Handles the level editor module starting to recompile.
 	void HandleLevelEditorModuleCompileStarted( bool bIsAsyncCompile );
 
@@ -330,15 +333,15 @@ private:
 	bool bIsReCreatingDefaultLayoutAsFallback = false;
 
 	// Bridge object for browser navigation
-	static TStrongObjectPtr<UHomeScreenWeb> WebObject;
+	TStrongObjectPtr<UHomeScreenWeb> WebObject = nullptr;
 
 	// Web browser window instance
-	static TSharedPtr<IWebBrowserWindow> WebBrowserWindow;
+	TSharedPtr<IWebBrowserWindow> WebBrowserWindow = nullptr;
 
 	// Web browser widget instance
-	static TSharedPtr<SWebBrowser> WebBrowser;
+	TSharedPtr<SWebBrowser> WebBrowser = nullptr;
 
 	// Holds the tab manager that manages the HomeScreen tab.
-    static TSharedPtr<FTabManager> HomeScreenTabManager;
+	TSharedPtr<FTabManager> HomeScreenTabManager = nullptr;
 
 };
