Added functions to MainFrameModule for opening or creating the ProjectBrowser that can be overriden.

[RN]

diff --git a/Engine/Source/Editor/MainFrame/Private/HomeScreen/SHomeScreen.cpp b/Engine/Source/Editor/MainFrame/Private/HomeScreen/SHomeScreen.cpp
--- a/Engine/Source/Editor/MainFrame/Private/HomeScreen/SHomeScreen.cpp
+++ b/Engine/Source/Editor/MainFrame/Private/HomeScreen/SHomeScreen.cpp
@@ -1,4 +1,4 @@
-﻿// Copyright Epic Games, Inc. All Rights Reserved.
+// Copyright Epic Games, Inc. All Rights Reserved.
 
 #include "SHomeScreen.h"
 
@@ -162,7 +162,9 @@ void SHomeScreen::Construct(const FArguments& InArgs, const TSharedPtr<SWebBrows
 	CreateSocialMediaButtons(InstagramButton, TEXT("https://www.instagram.com/unrealengine/"), FAppStyle::GetBrush("SocialMedia.Instagram"));
 
 	// Project Browser
-	ProjectBrowser = FGameProjectGenerationModule::Get().CreateProjectBrowser(false, false);
+	constexpr bool bProjectBrowserShowHeaderBar = false;
+	constexpr bool bProjectBrowserAllowScrolling = false;
+	ProjectBrowser = IMainFrameModule::Get().CreateProjectBrowser(bProjectBrowserShowHeaderBar, bProjectBrowserAllowScrolling);
 
 	// Check for internet connection on startup
 	CheckInternetConnection();
@@ -638,25 +640,7 @@ void SHomeScreen::Construct(const FArguments& InArgs, const TSharedPtr<SWebBrows
 
 FReply SHomeScreen::OnCreateProjectDialog(bool bInAllowProjectOpening, bool bInAllowProjectCreation)
 {
-	if (GUnrealEd->WarnIfLightingBuildIsCurrentlyRunning())
-	{
-		return FReply::Handled();
-	}
-
-	TSharedRef<SWindow> NewProjectWindow =
-		FMainFrameActionCallbacks::CreateProjectBrowserWindow(bInAllowProjectOpening, bInAllowProjectCreation);
-
-	NewProjectWindow->SetContent(FGameProjectGenerationModule::Get().CreateGameProjectDialog(bInAllowProjectOpening, bInAllowProjectCreation));
-
-	IMainFrameModule& MainFrameModule = FModuleManager::GetModuleChecked<IMainFrameModule>(TEXT("MainFrame"));
-	if (MainFrameModule.GetParentWindow().IsValid())
-	{
-		FSlateApplication::Get().AddWindowAsNativeChild(NewProjectWindow, MainFrameModule.GetParentWindow().ToSharedRef());
-	}
-	else
-	{
-		FSlateApplication::Get().AddWindow(NewProjectWindow);
-	}
+	IMainFrameModule::Get().ShowProjectBrowser(bInAllowProjectOpening, bInAllowProjectCreation);
 
 	return FReply::Handled();
 }
diff --git a/Engine/Source/Editor/MainFrame/Private/MainFrameModule.cpp b/Engine/Source/Editor/MainFrame/Private/MainFrameModule.cpp
--- a/Engine/Source/Editor/MainFrame/Private/MainFrameModule.cpp
+++ b/Engine/Source/Editor/MainFrame/Private/MainFrameModule.cpp
@@ -632,6 +632,27 @@ void FMainFrameModule::SetApplicationTitleOverride(const FText& NewOverriddenApp
 /* IModuleInterface implementation
  *****************************************************************************/
 
+void FMainFrameModule::ShowProjectBrowser(bool bAllowProjectOpening, bool bAllowProjectCreate)
+{
+	if (ShowProjectBrowserOverride.IsBound())
+	{
+		ShowProjectBrowserOverride.Execute(bAllowProjectOpening, bAllowProjectCreate);
+		return;
+	}
+
+	FMainFrameActionCallbacks::NewProject(bAllowProjectOpening, bAllowProjectCreate);
+}
+
+TSharedRef<SWidget> FMainFrameModule::CreateProjectBrowser(bool bInShowHeaderBar, bool bInAllowScrolling)
+{
+	if (CreateProjectBrowserOverride.IsBound())
+	{
+		return CreateProjectBrowserOverride.Execute(bInShowHeaderBar, bInAllowScrolling);
+	}
+
+	return FGameProjectGenerationModule::Get().CreateProjectBrowser(bInShowHeaderBar, bInAllowScrolling);
+}
+
 void FMainFrameModule::StartupModule( )
 {
 	bDelayedShowMainFrame = false;
diff --git a/Engine/Source/Editor/MainFrame/Private/MainFrameModule.h b/Engine/Source/Editor/MainFrame/Private/MainFrameModule.h
--- a/Engine/Source/Editor/MainFrame/Private/MainFrameModule.h
+++ b/Engine/Source/Editor/MainFrame/Private/MainFrameModule.h
@@ -208,7 +208,43 @@ public:
 		OutSectionName = EditorSettingsDefaultSectionOverride;
 	}
 
-public:
+	// Opens the project browser in a new window
+	virtual void ShowProjectBrowser(bool bAllowProjectOpening, bool bAllowProjectCreate) override;
+
+	// Creates the project browser for embedding
+	virtual TSharedRef<SWidget> CreateProjectBrowser(bool bInShowHeaderBar = true, bool bInAllowScrolling = true) override;
+
+	// Overrides the functionality of ShowProjectBrowser with the provided delegate
+	virtual FDelegateHandle RegisterShowProjectBrowserOverride(const FShowProjectBrowser& InShowProjectBrowserOverride) override
+	{
+		ShowProjectBrowserOverride = InShowProjectBrowserOverride;
+		return ShowProjectBrowserOverride.GetHandle();
+	}
+
+	// Removes the override of ShowProjectBrowser if it matches the handle
+	virtual void UnregisterShowProjectBrowserOverride(FDelegateHandle InHandle) override
+	{
+		if (ShowProjectBrowserOverride.GetHandle() == InHandle)
+		{
+			ShowProjectBrowserOverride = FShowProjectBrowser{};
+		}
+	}
+
+	// Overrides the functionality of CreateProjectBrowser with the provided delegate
+	virtual FDelegateHandle RegisterCreateProjectBrowserOverride(const FCreateProjectBrowser& InCreateProjectBrowserOverride) override
+	{
+		CreateProjectBrowserOverride = InCreateProjectBrowserOverride;
+		return CreateProjectBrowserOverride.GetHandle();
+	}
+
+	// Removes the override of CreateProjectBrowser if it matches the handle
+	virtual void UnregisterCreateProjectBrowserOverride(FDelegateHandle InHandle) override
+	{
+		if (CreateProjectBrowserOverride.GetHandle() == InHandle)
+		{
+			CreateProjectBrowserOverride = FCreateProjectBrowser{};
+		}
+	}
 
 	// IModuleInterface interface
 
@@ -344,4 +380,9 @@ private:
 	// Holds the tab manager that manages the HomeScreen tab.
 	TSharedPtr<FTabManager> HomeScreenTabManager = nullptr;
 
+	// Stores the delegate for the override of ShowProjectBrowser, if any
+	FShowProjectBrowser ShowProjectBrowserOverride;
+
+	// Stores the delegate for the override of CreateProjectBrowser, if any
+	FCreateProjectBrowser CreateProjectBrowserOverride;
 };
diff --git a/Engine/Source/Editor/MainFrame/Public/Interfaces/IMainFrameModule.h b/Engine/Source/Editor/MainFrame/Public/Interfaces/IMainFrameModule.h
--- a/Engine/Source/Editor/MainFrame/Public/Interfaces/IMainFrameModule.h
+++ b/Engine/Source/Editor/MainFrame/Public/Interfaces/IMainFrameModule.h
@@ -27,6 +27,9 @@ struct FMainFrameDeveloperTool
 	TAttribute<FText> Value;
 };
 
+FUNC_DECLARE_DELEGATE(FShowProjectBrowser, void, bool /*bAllowProjectOpening*/, bool /*bAllowProjectCreate*/);
+FUNC_DECLARE_DELEGATE(FCreateProjectBrowser, TSharedRef<SWidget>, bool /*bInShowHeaderBar*/, bool /*bInAllowScrolling*/);
+
 /**
  * Interface for main frame modules.
  */
@@ -228,6 +231,36 @@ public:
 	 */
 	virtual void GetEditorSettingsDefaultSelectionOverride(FName& OutCategoryName, FName& OutSectionName) = 0;
 
+	/**
+	 * Opens the project browser in a new window.
+	 */
+	virtual void ShowProjectBrowser(bool bAllowProjectOpening, bool bAllowProjectCreate) = 0;
+
+	/**
+	 * Creates the project browser for embedding.
+	 */
+	virtual TSharedRef<SWidget> CreateProjectBrowser(bool bInShowHeaderBar = true, bool bInAllowScrolling = true) = 0;
+
+	/**
+	 * Overrides the functionality of ShowProjectBrowser with the provided delegate.
+	 */
+	virtual FDelegateHandle RegisterShowProjectBrowserOverride(const FShowProjectBrowser& InShowProjectBrowserOverride) = 0;
+
+	/**
+	 * Removes the override of ShowProjectBrowser if it matches the handle.
+	 */
+	virtual void UnregisterShowProjectBrowserOverride(FDelegateHandle InHandle) = 0;
+
+	/**
+	 * Overrides the functionality of CreateProjectBrowser with the provided delegate.
+	 */
+	virtual FDelegateHandle RegisterCreateProjectBrowserOverride(const FCreateProjectBrowser& InCreateProjectBrowserOverride) = 0;
+
+	/**
+	 * Removes the override of CreateProjectBrowser if it matches the handle.
+	 */
+	virtual void UnregisterCreateProjectBrowserOverride(FDelegateHandle InHandle) = 0;
+
 public:
 
 	/**
