[HomeScreen] Add home screen customization

diff --git a/Engine/Plugins/Enterprise/DatasmithContent/Source/DatasmithContent/Private/DatasmithImportOptions.cpp b/Engine/Plugins/Enterprise/DatasmithContent/Source/DatasmithContent/Private/DatasmithImportOptions.cpp
--- a/Engine/Plugins/Enterprise/DatasmithContent/Source/DatasmithContent/Private/DatasmithImportOptions.cpp
+++ b/Engine/Plugins/Enterprise/DatasmithContent/Source/DatasmithContent/Private/DatasmithImportOptions.cpp
@@ -5,7 +5,6 @@
 #include "UObject/UnrealType.h"
 
 #include "HAL/FileManager.h"
-#include "Interfaces/IMainFrameModule.h"
 
 #include "Dom/JsonValue.h"
 #include "Dom/JsonObject.h"
diff --git a/Engine/Source/Editor/EditorStyle/Private/StarshipStyle.cpp b/Engine/Source/Editor/EditorStyle/Private/StarshipStyle.cpp
--- a/Engine/Source/Editor/EditorStyle/Private/StarshipStyle.cpp
+++ b/Engine/Source/Editor/EditorStyle/Private/StarshipStyle.cpp
@@ -8519,6 +8519,8 @@ void FStarshipEditorStyle::FStyle::SetupHomeScreenStyles()
 	Set("HomeScreen.Tutorial", new IMAGE_BRUSH_SVG("Starship/HomeScreen/Tutorial", Icon20x20));
 	Set("HomeScreen.Roadmap", new IMAGE_BRUSH_SVG("Starship/HomeScreen/Roadmap", Icon20x20));
 	Set("HomeScreen.ReleaseNotes", new IMAGE_BRUSH_SVG("Starship/HomeScreen/ReleaseNotes", Icon20x20));
+	Set("HomeScreen.CreatorRules", new IMAGE_BRUSH_SVG("Starship/HomeScreen/ParentChild", Icon20x20));
+	Set("HomeScreen.Support", new IMAGE_BRUSH_SVG("Starship/HomeScreen/HeadphonesMic", Icon20x20));
 	Set("HomeScreen.PlusCircle", new IMAGE_BRUSH_SVG("Starship/HomeScreen/PlusCircleEmpty", Icon16x16));
 	Set("HomeScreen.FolderOpen", new IMAGE_BRUSH_SVG("Starship/HomeScreen/FolderOpen", Icon16x16));
 
diff --git a/Engine/Source/Editor/MainFrame/Private/HomeScreen/SHomeScreen.cpp b/Engine/Source/Editor/MainFrame/Private/HomeScreen/SHomeScreen.cpp
--- a/Engine/Source/Editor/MainFrame/Private/HomeScreen/SHomeScreen.cpp
+++ b/Engine/Source/Editor/MainFrame/Private/HomeScreen/SHomeScreen.cpp
@@ -110,6 +110,7 @@ void SHomeScreen::Construct(const FArguments& InArgs, const TSharedPtr<SWebBrows
 {
 	WebBrowser = InWebBrowser;
 	WebObject = InWebObject;
+	SectionLinks = InArgs._SectionLinks;
 
 	AutoLoadProjectComboBoxSelection = EAutoLoadProject::HomeScreen;
 	if (UHomeScreenSettings* HomeScreenSettings = GetMutableDefault<UHomeScreenSettings>())
@@ -136,18 +137,6 @@ void SHomeScreen::Construct(const FArguments& InArgs, const TSharedPtr<SWebBrows
 	CreateMainSectionCheckBox(GettingStartedCheckBox, EMainSectionMenu::GettingStarted, LOCTEXT("ProjectHomeGettingStarted", "Getting Started"), FAppStyle::GetBrush("HomeScreen.Rocket"));
 	CreateMainSectionCheckBox(SampleProjectsCheckBox, EMainSectionMenu::SampleProjects, LOCTEXT("ProjectHomeSampleProjects", "Sample Projects"), FAppStyle::GetBrush("HomeScreen.Archive"));
 
-	// Documentation and Community Buttons
-	CreateResourceButtons(ForumsButton, TEXT("https://forums.unrealengine.com/categories?tag=unreal-engine"),
-		LOCTEXT("ProjectHomeForumsButton", "Forums"), FAppStyle::GetBrush("HomeScreen.Forum"));
-	CreateResourceButtons(DocumentationButton, TEXT("https://dev.epicgames.com/documentation/en-us/unreal-engine/unreal-engine-5-6-documentation"),
-		LOCTEXT("ProjectHomeDocumentation", "Documentation"), FAppStyle::GetBrush("HomeScreen.Documentation"));
-	CreateResourceButtons(TutorialsButton, TEXT("https://dev.epicgames.com/community/unreal-engine/learning"),
-		LOCTEXT("ProjectHomeTutorialsButton", "Tutorials"), FAppStyle::GetBrush("HomeScreen.Tutorial"));
-	CreateResourceButtons(RoadmapButton, TEXT("https://portal.productboard.com/epicgames/1-unreal-engine-public-roadmap/tabs/109-unreal-engine-5-5"),
-		LOCTEXT("ProjectHomeRoadmapButton", "Roadmap"), FAppStyle::GetBrush("HomeScreen.Roadmap"));
-	CreateResourceButtons(ReleaseNotesButton, TEXT("https://dev.epicgames.com/documentation/en-us/unreal-engine/unreal-engine-5-6-release-notes"),
-		LOCTEXT("ProjectHomeReleaseNotesButton", "Release Notes"), FAppStyle::GetBrush("HomeScreen.ReleaseNotes"));
-
 	// Social Media Buttons
 	TSharedPtr<SButton> FacebookButton;
 	TSharedPtr<SButton> TwitterButton;
@@ -323,47 +312,7 @@ void SHomeScreen::Construct(const FArguments& InArgs, const TSharedPtr<SWebBrows
 					.AutoHeight()
 					.Padding(16.f, 0.f)
 					[
-						SNew(SVerticalBox)
-
-						+ SVerticalBox::Slot()
-						.AutoHeight()
-						.MinHeight(SectionButtonCheckBoxHeight)
-						.MaxHeight(SectionButtonCheckBoxHeight)
-						[
-							ForumsButton.ToSharedRef()
-						]
-
-						+ SVerticalBox::Slot()
-						.AutoHeight()
-						.MinHeight(SectionButtonCheckBoxHeight)
-						.MaxHeight(SectionButtonCheckBoxHeight)
-						[
-							DocumentationButton.ToSharedRef()
-						]
-
-						+ SVerticalBox::Slot()
-						.AutoHeight()
-						.MinHeight(SectionButtonCheckBoxHeight)
-						.MaxHeight(SectionButtonCheckBoxHeight)
-						[
-							TutorialsButton.ToSharedRef()
-						]
-
-						+ SVerticalBox::Slot()
-						.AutoHeight()
-						.MinHeight(SectionButtonCheckBoxHeight)
-						.MaxHeight(SectionButtonCheckBoxHeight)
-						[
-							RoadmapButton.ToSharedRef()
-						]
-
-						+ SVerticalBox::Slot()
-						.AutoHeight()
-						.MinHeight(SectionButtonCheckBoxHeight)
-						.MaxHeight(SectionButtonCheckBoxHeight)
-						[
-							ReleaseNotesButton.ToSharedRef()
-						]
+						SAssignNew(ResourceAndCommunityVerticalBox, SVerticalBox)
 					]
 
 					// Load on Startup Section
@@ -635,6 +584,19 @@ void SHomeScreen::Construct(const FArguments& InArgs, const TSharedPtr<SWebBrows
 		]
 	];
 
+	for (const FHomeScreenResourceCommunityEntry& ResourceCommunityEntries : InArgs._ResourceCommunityEntries)
+	{
+		TSharedPtr<SButton> OutButton;
+		CreateResourceButtons(OutButton, ResourceCommunityEntries.Link, ResourceCommunityEntries.DisplayName, FAppStyle::GetBrush(ResourceCommunityEntries.IconBrushName));
+		ResourceAndCommunityVerticalBox->AddSlot()
+			.AutoHeight()
+			.MinHeight(SectionButtonCheckBoxHeight)
+			.MaxHeight(SectionButtonCheckBoxHeight)
+			[
+				OutButton.ToSharedRef()
+			];
+	}
+
 	OnMainHomeSectionChanged(ECheckBoxState::Checked, MainHomeSelection);
 }
 
@@ -658,35 +620,31 @@ ECheckBoxState SHomeScreen::IsMainHomeSectionChecked(EMainSectionMenu InMainHome
 
 void SHomeScreen::OnMainHomeSectionChanged(ECheckBoxState InCheckBoxState, EMainSectionMenu InMainHomeSelection)
 {
-	FString URL = TEXT("");
-	const FString BaseURL = TEXT("https://editor.unrealengine.com");
+	// If this HomeScreen version do not have a link for the requested section menu do not process it
+	if (!SectionLinks.Contains(InMainHomeSelection))
+	{
+		return;
+	}
+
+	MainHomeSelection = InMainHomeSelection;
+
+	const FHomeScreenSectionLink& CurrentRequestedLink = SectionLinks[InMainHomeSelection];
+	const FString BaseURL = CurrentRequestedLink.BeforeLocaleURL;
+
 	// Use the Editor Language for the WebPage language.
 	const FString FullLocale = FInternationalization::Get().GetCurrentLanguage()->GetName();
-	switch (InMainHomeSelection)
+	FString EndURL = CurrentRequestedLink.AfterLocaleURL;
+
+	if (!bHasLatestEngineProjects && !CurrentRequestedLink.AfterLocaleNoProjectURL.IsEmpty())
 	{
-	case EMainSectionMenu::News:
-		MainHomeSelection = EMainSectionMenu::News;
-		URL = BaseURL / FullLocale / TEXT("news");
-		break;
-	case EMainSectionMenu::GettingStarted:
-		MainHomeSelection = EMainSectionMenu::GettingStarted;
-		URL = BaseURL / FullLocale / TEXT("get-started-with-unreal-engine");
-		break;
-	case EMainSectionMenu::SampleProjects:
-		MainHomeSelection = EMainSectionMenu::SampleProjects;
-		URL = BaseURL / FullLocale / TEXT("samples");
-		break;
-	case EMainSectionMenu::None:
-	case EMainSectionMenu::Home:
-	default:
-		MainHomeSelection = EMainSectionMenu::Home;
-		URL = bHasLatestEngineProjects ? BaseURL / FullLocale / TEXT("recent-projects") : BaseURL / FullLocale / TEXT("no-projects");
-		break;
+		EndURL = CurrentRequestedLink.AfterLocaleNoProjectURL;
 	}
 
+	const FString FinalURL = BaseURL / FullLocale / EndURL;
+
 	if (WebBrowser.IsValid())
 	{
-		WebBrowser->LoadURL(URL);
+		WebBrowser->LoadURL(FinalURL);
 	}
 }
 
@@ -752,6 +710,7 @@ FSlateColor SHomeScreen::GetResourceAndSocialMediaButtonColor(const TSharedPtr<S
 void SHomeScreen::CreateMainSectionCheckBox(TSharedPtr<SCheckBox>& OutCheckBox, EMainSectionMenu InMainHomeSelection, const FText& InText, const FSlateBrush* InImage)
 {
 	OutCheckBox = SNew(SCheckBox)
+		.Visibility_Lambda([this, InMainHomeSelection] () { return SectionLinks.Contains(InMainHomeSelection) ? EVisibility::Visible : EVisibility::Collapsed; })
 		.IsEnabled(this, &SHomeScreen::IsMainSectionEnabled, InMainHomeSelection)
 		.Style(FAppStyle::Get(), "HomeScreen.MainMenuSectionButton")
 		.IsChecked(this, &SHomeScreen::IsMainHomeSectionChecked, InMainHomeSelection)
diff --git a/Engine/Source/Editor/MainFrame/Private/HomeScreen/SHomeScreen.h b/Engine/Source/Editor/MainFrame/Private/HomeScreen/SHomeScreen.h
--- a/Engine/Source/Editor/MainFrame/Private/HomeScreen/SHomeScreen.h
+++ b/Engine/Source/Editor/MainFrame/Private/HomeScreen/SHomeScreen.h
@@ -13,6 +13,7 @@
 class SButton;
 class SCheckBox;
 class SComboButton;
+class SVerticalBox;
 class SWebBrowser;
 template <typename ItemType> class STileView;
 
@@ -20,8 +21,11 @@ class SHomeScreen : public SCompoundWidget
 {
 public:
 	SLATE_BEGIN_ARGS(SHomeScreen)
+		: _ResourceCommunityEntries()
+		, _SectionLinks()
 	{}
-
+		SLATE_ARGUMENT(TArray<FHomeScreenResourceCommunityEntry>, ResourceCommunityEntries)
+		SLATE_ARGUMENT(FHomeScreenSectionToLinkMap, SectionLinks)
 	SLATE_END_ARGS()
 
 	SHomeScreen();
@@ -114,13 +118,6 @@ private:
 	TSharedPtr<SCheckBox> GettingStartedCheckBox;
 	TSharedPtr<SCheckBox> SampleProjectsCheckBox;
 
-	/** Resource section buttons */
-	TSharedPtr<SButton> ForumsButton;
-	TSharedPtr<SButton> DocumentationButton;
-	TSharedPtr<SButton> TutorialsButton;
-	TSharedPtr<SButton> RoadmapButton;
-	TSharedPtr<SButton> ReleaseNotesButton;
-
 	/** Http retry system used to check Internet connection */
 	TSharedPtr<FHttpRetrySystem::FManager> HttpRetryManager;
 	TSharedPtr<FHttpRetrySystem::FRequest> HttpRetryRequest;
@@ -134,6 +131,9 @@ private:
 	/** Current combo button menu */
 	TWeakPtr<SWidget> ComboButtonMenuWeak;
 
+	/** VerticalBox holding the ResourceAndCommunity buttons */
+	TSharedPtr<SVerticalBox> ResourceAndCommunityVerticalBox;
+
 	/** Timer handle for checking internet connection */
 	FTimerHandle CheckInternetConnectionTimerHandle;
 
@@ -166,4 +166,7 @@ private:
 	
 	/** WebBrowser object passed down from the MainFrameModule */
 	TWeakObjectPtr<UHomeScreenWeb> WebObject;
+	
+	/** Section x link combination */
+	FHomeScreenSectionToLinkMap SectionLinks;
 };
diff --git a/Engine/Source/Editor/MainFrame/Private/MainFrameModule.cpp b/Engine/Source/Editor/MainFrame/Private/MainFrameModule.cpp
--- a/Engine/Source/Editor/MainFrame/Private/MainFrameModule.cpp
+++ b/Engine/Source/Editor/MainFrame/Private/MainFrameModule.cpp
@@ -92,22 +92,119 @@ const FText StaticGetApplicationTitle( const bool bIncludeGameName )
 	return ApplicationTitle;
 }
 
+namespace HomeScreenHelper
+{
+	// Get the default UE home screen resource entries data
+	TArray<FHomeScreenResourceCommunityEntry> GetUEResourceEntries()
+	{
+		TArray<FHomeScreenResourceCommunityEntry> DefaultEntries;
+
+		// Forum
+		FHomeScreenResourceCommunityEntry ForumEntry;
+		ForumEntry.DisplayName = LOCTEXT("ProjectHomeForumsButton", "Forums");
+		ForumEntry.Link = TEXT("https://forums.unrealengine.com/categories?tag=unreal-engine");
+		ForumEntry.IconBrushName = TEXT("HomeScreen.Forum");
+		DefaultEntries.Add(ForumEntry);
+
+		// Documentation
+		FHomeScreenResourceCommunityEntry DocumentationEntry;
+		DocumentationEntry.DisplayName = LOCTEXT("ProjectHomeDocumentation", "Documentation");
+		DocumentationEntry.Link = TEXT("https://dev.epicgames.com/documentation/en-us/unreal-engine/unreal-engine-5-6-documentation");
+		DocumentationEntry.IconBrushName = TEXT("HomeScreen.Documentation");
+		DefaultEntries.Add(DocumentationEntry);
+
+		// Tutorial
+		FHomeScreenResourceCommunityEntry TutorialEntry;
+		TutorialEntry.DisplayName = LOCTEXT("ProjectHomeTutorialsButton", "Tutorials");
+		TutorialEntry.Link = TEXT("https://dev.epicgames.com/community/unreal-engine/learning");
+		TutorialEntry.IconBrushName = TEXT("HomeScreen.Tutorial");
+		DefaultEntries.Add(TutorialEntry);
+
+		// Roadmap
+		FHomeScreenResourceCommunityEntry RoadmapEntry;
+		RoadmapEntry.DisplayName = LOCTEXT("ProjectHomeRoadmapButton", "Roadmap");
+		RoadmapEntry.Link = TEXT("https://portal.productboard.com/epicgames/1-unreal-engine-public-roadmap/tabs/109-unreal-engine-5-5");
+		RoadmapEntry.IconBrushName = TEXT("HomeScreen.Roadmap");
+		DefaultEntries.Add(RoadmapEntry);
+
+		// Release
+		FHomeScreenResourceCommunityEntry ReleaseEntry;
+		ReleaseEntry.DisplayName = LOCTEXT("ProjectHomeReleaseNotesButton", "Release Notes");
+		ReleaseEntry.Link = TEXT("https://dev.epicgames.com/documentation/en-us/unreal-engine/unreal-engine-5-6-release-notes");
+		ReleaseEntry.IconBrushName = TEXT("HomeScreen.ReleaseNotes");
+		DefaultEntries.Add(ReleaseEntry);
+		
+		return DefaultEntries;
+	}
+
+	// Get the default UE home screen main section entries data
+	FHomeScreenSectionToLinkMap GetUESectionLinks()
+	{
+		FHomeScreenSectionToLinkMap OutSectionLinks;
+		const FString EmptyURL = TEXT("");
+		const FString BaseURL = TEXT("https://editor.unrealengine.com");
+		OutSectionLinks.Add(EMainSectionMenu::News, 
+			{ BaseURL,  TEXT("news"), EmptyURL});
+		OutSectionLinks.Add(EMainSectionMenu::GettingStarted, 
+			{ BaseURL,  TEXT("get-started-with-unreal-engine"), EmptyURL});
+		OutSectionLinks.Add(EMainSectionMenu::SampleProjects, 
+			{ BaseURL,  TEXT("samples"), EmptyURL});
+		OutSectionLinks.Add(EMainSectionMenu::Home, 
+			{ BaseURL,  TEXT("recent-projects"), TEXT("no-projects")});
+		return OutSectionLinks;
+	}
+
+	// Get the default UE home screen settings
+	FHomeScreenWidgetSettings GetHomeScreenSettings()
+	{
+		FHomeScreenWidgetSettings HomeScreenWidgetSettings;
+		HomeScreenWidgetSettings.ResourcesEntries = GetUEResourceEntries();
+		HomeScreenWidgetSettings.SectionLinks = GetUESectionLinks();
+		return HomeScreenWidgetSettings;
+	}
+}
+
+/* FEditorHomeScreenProvider
+ *****************************************************************************/
+
+void FEditorHomeScreenProvider::Register()
+{
+	IModularFeatures::Get().RegisterModularFeature(IHomeScreenProvider::GetModularFeatureName(), this);
+}
+
+void FEditorHomeScreenProvider::UnRegister()
+{
+	IModularFeatures::Get().UnregisterModularFeature(IHomeScreenProvider::GetModularFeatureName(), this);
+}
+
+bool FEditorHomeScreenProvider::IsRequestingHomeScreenPageControl() const
+{
+	return true;
+}
+
+int32 FEditorHomeScreenProvider::GetHomeScreenPriority()
+{
+	// Lower priority since this is the default behavior
+	return -1;
+}
+
+FHomeScreenWidgetSettings FEditorHomeScreenProvider::GetHomeScreenWidgetSettings() const
+{
+	// No need to add settings as we want them to be the default
+	return HomeScreenHelper::GetHomeScreenSettings();
+}
+
 /* FProjectDialogProvider
  *****************************************************************************/
 
-void FProjectDialogProvider::Register(FMainFrameModule* InMainFrameModule)
+void FProjectDialogProvider::Register()
 {
 	IModularFeatures::Get().RegisterModularFeature(IEditorMainFrameProvider::GetModularFeatureName(), this);
-	if (InMainFrameModule != nullptr)
-	{
-		MainFrameModulePtr = InMainFrameModule;
-	}
 }
 
 void FProjectDialogProvider::UnRegister()
 {
 	IModularFeatures::Get().UnregisterModularFeature(IEditorMainFrameProvider::GetModularFeatureName(), this);
-	MainFrameModulePtr = nullptr;
 }
 
 bool FProjectDialogProvider::IsRequestingMainFrameControl() const
@@ -150,10 +247,7 @@ TSharedRef<SWidget> FProjectDialogProvider::CreateMainFrameContentWidget() const
 {
 	if (UE::Editor::HomeScreen::IsHomeScreenEnabled())
 	{
-		if (MainFrameModulePtr != nullptr)
-		{
-			return MainFrameModulePtr->CreateHomeScreenWidget({});
-		}
+		return IMainFrameModule::Get().CreateHomeScreenWidget({}, HomeScreenHelper::GetHomeScreenSettings());
 	}
 
 	return FGameProjectGenerationModule::Get().CreateGameProjectDialog(/*bAllowProjectOpening=*/true, /*bAllowProjectCreate=*/true);
@@ -660,7 +754,8 @@ void FMainFrameModule::StartupModule( )
 
 	MRUFavoritesList = NULL;
 
-	ProjectDialogProvider.Register(this);
+	ProjectDialogProvider.Register();
+	HomeScreenProvider.Register();
 
 	ensureMsgf(!IsRunningGame(), TEXT("The MainFrame module should only be loaded when running the editor.  Code that extends the editor, adds menu items, etc... should not run when running in -game mode or in a non-WITH_EDITOR build"));
 	MainFrameHandler = MakeShareable(new FMainFrameHandler);
@@ -692,7 +787,7 @@ void FMainFrameModule::StartupModule( )
 			.SetAutoGenerateMenuEntry(false)
 			.SetMenuType(ETabSpawnerMenuType::Hidden)
 			.SetDisplayName(FText::GetEmpty());
-    	}
+    }
 
 	ModuleCompileStartTime = 0.0f;
 
@@ -742,6 +837,7 @@ void FMainFrameModule::ShutdownModule( )
 	FGlobalTabmanager::Get()->UnregisterTabSpawner(TEXT("HomeScreen"));
 
 	ProjectDialogProvider.UnRegister();
+	HomeScreenProvider.UnRegister();
 }
 
 
@@ -958,13 +1054,36 @@ TSharedRef<SDockTab> FMainFrameModule::SpawnHomeScreenWindow(const FSpawnTabArgs
 
 	HomeScreenTab->SetTabIcon(FAppStyle::Get().GetBrush("HomeScreen.Home"));
 
+	TArray<IHomeScreenProvider*> HomeScreenProviders = IModularFeatures::Get().GetModularFeatureImplementations<IHomeScreenProvider>(IHomeScreenProvider::GetModularFeatureName());		
+	IHomeScreenProvider* ActiveProvider = nullptr;
+	for (IHomeScreenProvider* Provider : HomeScreenProviders)
+	{
+		if (Provider && Provider->IsRequestingHomeScreenPageControl())
+		{
+			if (!ActiveProvider || Provider->GetHomeScreenPriority() > ActiveProvider->GetHomeScreenPriority())
+			{
+				ActiveProvider = Provider;
+			}
+		}
+	}
+
+	TSharedPtr<SWidget> HomeScreenWidget;
+	if (ActiveProvider)
+	{
+		HomeScreenWidget = CreateHomeScreenWidget(HomeScreenTab, ActiveProvider->GetHomeScreenWidgetSettings());
+	}
+	else
+	{
+		HomeScreenWidget = CreateHomeScreenWidget(HomeScreenTab);
+	}
+
 	// This will create the necessary object and widgets needed for the HomeScreen and return a newly created one or NullWidget if something went wrong
-	HomeScreenTab->SetContent(CreateHomeScreenWidget(HomeScreenTab));
+	HomeScreenTab->SetContent(HomeScreenWidget.ToSharedRef());
 
 	return HomeScreenTab;
 }
 
-TSharedRef<SWidget> FMainFrameModule::CreateHomeScreenWidget(const TSharedPtr<SDockTab>& InParentDockTab)
+TSharedRef<SWidget> FMainFrameModule::CreateHomeScreenWidget(const TSharedPtr<SDockTab>& InParentDockTab, FHomeScreenWidgetSettings InSettings)
 {
 	if (!WebObject.IsValid())
 	{
@@ -1038,7 +1157,9 @@ TSharedRef<SWidget> FMainFrameModule::CreateHomeScreenWidget(const TSharedPtr<SD
 		WebBrowserWindow->Reload();
 	}
 
-	return SNew(SHomeScreen, WebBrowser, TWeakObjectPtr(WebObject.Get()));
+	return SNew(SHomeScreen, WebBrowser, TWeakObjectPtr(WebObject.Get()))
+		.SectionLinks(InSettings.SectionLinks)
+		.ResourceCommunityEntries(InSettings.ResourcesEntries);
 }
 
 void FMainFrameModule::HandleCodeAccessorOpenFileFailed(const FString& Filename)
diff --git a/Engine/Source/Editor/MainFrame/Private/MainFrameModule.h b/Engine/Source/Editor/MainFrame/Private/MainFrameModule.h
--- a/Engine/Source/Editor/MainFrame/Private/MainFrameModule.h
+++ b/Engine/Source/Editor/MainFrame/Private/MainFrameModule.h
@@ -12,11 +12,29 @@
 #include "Frame/MainFrameHandler.h"
 #include "Misc/CompilationResult.h"
 #include "Interfaces/IEditorMainFrameProvider.h"
+#include "Interfaces/IHomeScreenProvider.h"
 
 class SWebBrowser;
 class IWebBrowserWindow;
 class UHomeScreenWeb;
 
+/** 
+ * Utility class, which hooks into the editor's home screen tab and
+ * supply it with its own settings if it is the higher priority provider
+ */
+class FEditorHomeScreenProvider : public IHomeScreenProvider
+{
+public:
+	void Register();
+	void UnRegister();
+
+	//~ Begin IHomeScreenProvider interface
+	virtual bool IsRequestingHomeScreenPageControl() const override;
+	virtual int32 GetHomeScreenPriority() override;
+	virtual FHomeScreenWidgetSettings GetHomeScreenWidgetSettings() const override;
+	//~ End IHomeScreenProvider interface
+};
+
 /** 
  * Utility class, which hooks into the editor's main window startup and 
  * supplies a project dialog window if the project hasn't been set. 
@@ -24,7 +42,7 @@ class UHomeScreenWeb;
 class FProjectDialogProvider : public IEditorMainFrameProvider
 {
 public:
-	void Register(class FMainFrameModule* InMainFrameModule);
+	void Register();
 	void UnRegister();
 
 	//~ Begin IEditorMainFrameProvider interface
@@ -32,9 +50,6 @@ public:
 	virtual FMainFrameWindowOverrides GetDesiredWindowConfiguration() const override;
 	virtual TSharedRef<SWidget> CreateMainFrameContentWidget() const override;
 	//~ End IEditorMainFrameProvider interface
-	
-private:
-	FMainFrameModule* MainFrameModulePtr = nullptr;
 };
 
 /**
@@ -264,7 +279,7 @@ public:
 	static FVector2D GetProjectBrowserWindowSize() { return FVector2D(1206, 822); }
 
 	// Create the home screen widget
-	TSharedRef<SWidget> CreateHomeScreenWidget(const TSharedPtr<SDockTab>& InParentTab);
+	virtual TSharedRef<SWidget> CreateHomeScreenWidget(const TSharedPtr<SDockTab>& InParentTab, FHomeScreenWidgetSettings InSettings = FHomeScreenWidgetSettings()) override;
 
 private:
 
@@ -365,6 +380,9 @@ private:
 	// Instantiation of the object responsible for spawning the editor's project dialog on startup
 	FProjectDialogProvider ProjectDialogProvider;
 
+	// Instantiation of the object responsible for getting the home screen settings to use if active
+	FEditorHomeScreenProvider HomeScreenProvider;
+
 	// Set to true if we are in the middle of recreating the default layout because a previous load attempt failed
 	bool bIsReCreatingDefaultLayoutAsFallback = false;
 
diff --git a/Engine/Source/Editor/MainFrame/Public/Interfaces/IHomeScreenProvider.h b/Engine/Source/Editor/MainFrame/Public/Interfaces/IHomeScreenProvider.h
new file mode 100644
--- /dev/null
+++ b/Engine/Source/Editor/MainFrame/Public/Interfaces/IHomeScreenProvider.h
@@ -0,0 +1,40 @@
+// Copyright Epic Games, Inc. All Rights Reserved.
+
+#pragma once
+
+#include "Features/IModularFeature.h"
+#include "Settings/HomeScreenCommon.h"
+
+/** 
+ * Base feature class, which serves as a hook for external sources to override 
+ * the editor's HomeScreen window content.
+ */
+class IHomeScreenProvider : public IModularFeature
+{
+public:
+	virtual ~IHomeScreenProvider() {}
+
+	static FName GetModularFeatureName()
+	{
+		static FName FeatureName = FName(TEXT("HomeScreenProvider"));
+		return FeatureName;
+	}
+
+	/**
+	 * @return True if we want our HomeScreen settings to be used for the HomeScreen widget inside the HomeScreen tab
+	 */
+	virtual bool IsRequestingHomeScreenPageControl() const = 0;
+
+	/**
+	 * Less priority means that it will be ignored for higher ones.
+	 * The editor default HomeScreen priority is -1
+	 * 
+	 * @return The priority for this HomeScreen
+	 */
+	virtual int32 GetHomeScreenPriority() = 0;
+
+	/**
+	 * @return Settings to use for this HomeScreen
+	 */
+	virtual FHomeScreenWidgetSettings GetHomeScreenWidgetSettings() const = 0;
+};
diff --git a/Engine/Source/Editor/MainFrame/Public/Interfaces/IMainFrameModule.h b/Engine/Source/Editor/MainFrame/Public/Interfaces/IMainFrameModule.h
--- a/Engine/Source/Editor/MainFrame/Public/Interfaces/IMainFrameModule.h
+++ b/Engine/Source/Editor/MainFrame/Public/Interfaces/IMainFrameModule.h
@@ -9,6 +9,7 @@
 #include "Widgets/SWidget.h"
 #include "Framework/Commands/UICommandList.h"
 #include "Framework/MultiBox/MultiBoxExtender.h"
+#include "Settings/HomeScreenCommon.h"
 
 class FTabManager;
 class SWindow;
@@ -38,6 +39,16 @@ class IMainFrameModule
 {
 public:
 
+	/**
+	 * Create the HomeScreen widget, will use the default settings if no settings are provided and return the newly created HomeScreen
+	 * Should only be used when using the IEditorMainFrameProvider if you wish to use a HomeScreen widget with different settings
+	 * 
+	 * @param InParentTab OwnerTab of the HomeScreen, needed when the HomeScreen is created as a Tab to avoid MAC to render the WebPage when switching tabs
+	 * @param InSettings HomeScreenSettings to use to create the current HomeScreen
+	 * @return Return the newly created HomeScreen widget
+	 */
+	virtual TSharedRef<SWidget> CreateHomeScreenWidget(const TSharedPtr<SDockTab>& InParentTab, FHomeScreenWidgetSettings InSettings = FHomeScreenWidgetSettings()) = 0;
+
 	/**
 	 * Creates the default editor main frame
 	 *
diff --git a/Engine/Source/Editor/MainFrame/Public/Settings/HomeScreenCommon.h b/Engine/Source/Editor/MainFrame/Public/Settings/HomeScreenCommon.h
--- a/Engine/Source/Editor/MainFrame/Public/Settings/HomeScreenCommon.h
+++ b/Engine/Source/Editor/MainFrame/Public/Settings/HomeScreenCommon.h
@@ -1,6 +1,11 @@
 ﻿// Copyright Epic Games, Inc. All Rights Reserved.
 
 #pragma once
+#include "Containers/Array.h"
+#include "Containers/Map.h"
+#include "Containers/UnrealString.h"
+#include "Internationalization/Text.h"
+#include "UObject/NameTypes.h"
 #include "HomeScreenCommon.generated.h"
 
 UENUM()
@@ -20,3 +25,40 @@ enum class EAutoLoadProject
 	LastProject UMETA(DisplayName = "Most Recent Project"),
 	MAX UMETA(Hidden)
 };
+
+/**
+ * Struct containing information on the MainSections in the HomeScreen
+ */
+struct FHomeScreenSectionLink
+{
+	FString BeforeLocaleURL;
+	FString AfterLocaleURL;
+	FString AfterLocaleNoProjectURL;
+};
+
+using FHomeScreenSectionToLinkMap = TMap<EMainSectionMenu, FHomeScreenSectionLink>;
+
+/**
+ * Struct containing information of ResourceAndCommunity section entries in the HomeScreen
+ */
+struct FHomeScreenResourceCommunityEntry
+{
+	FString Link;
+	FText DisplayName;
+	FName IconBrushName;
+};
+
+/** Setting used to create the HomeScreen, if leaved empty the defaults are used instead */
+struct FHomeScreenWidgetSettings
+{
+	/** 
+	 * Section x Link connection, if leaved empty it will use the UE default section x links.
+	 * If at least 1 is populated it will use this instead, every missing section are skipped from being created in this case
+	 */
+	FHomeScreenSectionToLinkMap SectionLinks;
+
+	/**
+	 * Additional Resource entry to add to the default ones
+	 */
+	TArray<FHomeScreenResourceCommunityEntry> ResourcesEntries;
+};
