Change to FCreateProjectBrowser to make the first parameter based on whether to only show the project list, rather than showing the header or not.

diff --git a/Engine/Source/Editor/MainFrame/Private/HomeScreen/SHomeScreen.cpp b/Engine/Source/Editor/MainFrame/Private/HomeScreen/SHomeScreen.cpp
--- a/Engine/Source/Editor/MainFrame/Private/HomeScreen/SHomeScreen.cpp
+++ b/Engine/Source/Editor/MainFrame/Private/HomeScreen/SHomeScreen.cpp
@@ -151,9 +151,9 @@ void SHomeScreen::Construct(const FArguments& InArgs, const TSharedPtr<SWebBrows
 	CreateSocialMediaButtons(InstagramButton, TEXT("https://www.instagram.com/unrealengine/"), FAppStyle::GetBrush("SocialMedia.Instagram"));
 
 	// Project Browser
-	constexpr bool bProjectBrowserShowHeaderBar = false;
+	constexpr bool bOnlyShowList = true;
 	constexpr bool bProjectBrowserAllowScrolling = false;
-	ProjectBrowser = IMainFrameModule::Get().CreateProjectBrowser(bProjectBrowserShowHeaderBar, bProjectBrowserAllowScrolling);
+	ProjectBrowser = IMainFrameModule::Get().CreateProjectBrowser(bOnlyShowList, bProjectBrowserAllowScrolling);
 
 	// Check for internet connection on startup
 	CheckInternetConnection();
diff --git a/Engine/Source/Editor/MainFrame/Private/MainFrameModule.cpp b/Engine/Source/Editor/MainFrame/Private/MainFrameModule.cpp
--- a/Engine/Source/Editor/MainFrame/Private/MainFrameModule.cpp
+++ b/Engine/Source/Editor/MainFrame/Private/MainFrameModule.cpp
@@ -737,14 +737,14 @@ void FMainFrameModule::ShowProjectBrowser(bool bAllowProjectOpening, bool bAllow
 	FMainFrameActionCallbacks::NewProject(bAllowProjectOpening, bAllowProjectCreate);
 }
 
-TSharedRef<SWidget> FMainFrameModule::CreateProjectBrowser(bool bInShowHeaderBar, bool bInAllowScrolling)
+TSharedRef<SWidget> FMainFrameModule::CreateProjectBrowser(bool bOnlyShowList, bool bAllowScrolling)
 {
 	if (CreateProjectBrowserOverride.IsBound())
 	{
-		return CreateProjectBrowserOverride.Execute(bInShowHeaderBar, bInAllowScrolling);
+		return CreateProjectBrowserOverride.Execute(bOnlyShowList, bAllowScrolling);
 	}
 
-	return FGameProjectGenerationModule::Get().CreateProjectBrowser(bInShowHeaderBar, bInAllowScrolling);
+	return FGameProjectGenerationModule::Get().CreateProjectBrowser(!bOnlyShowList, bAllowScrolling);
 }
 
 void FMainFrameModule::StartupModule( )
diff --git a/Engine/Source/Editor/MainFrame/Private/MainFrameModule.h b/Engine/Source/Editor/MainFrame/Private/MainFrameModule.h
--- a/Engine/Source/Editor/MainFrame/Private/MainFrameModule.h
+++ b/Engine/Source/Editor/MainFrame/Private/MainFrameModule.h
@@ -227,7 +227,7 @@ public:
 	virtual void ShowProjectBrowser(bool bAllowProjectOpening, bool bAllowProjectCreate) override;
 
 	// Creates the project browser for embedding
-	virtual TSharedRef<SWidget> CreateProjectBrowser(bool bInShowHeaderBar = true, bool bInAllowScrolling = true) override;
+	virtual TSharedRef<SWidget> CreateProjectBrowser(bool bOnlyShowList = true, bool bAllowScrolling = true) override;
 
 	// Overrides the functionality of ShowProjectBrowser with the provided delegate
 	virtual FDelegateHandle RegisterShowProjectBrowserOverride(const FShowProjectBrowser& InShowProjectBrowserOverride) override
diff --git a/Engine/Source/Editor/MainFrame/Public/Interfaces/IMainFrameModule.h b/Engine/Source/Editor/MainFrame/Public/Interfaces/IMainFrameModule.h
--- a/Engine/Source/Editor/MainFrame/Public/Interfaces/IMainFrameModule.h
+++ b/Engine/Source/Editor/MainFrame/Public/Interfaces/IMainFrameModule.h
@@ -29,7 +29,7 @@ struct FMainFrameDeveloperTool
 };
 
 FUNC_DECLARE_DELEGATE(FShowProjectBrowser, void, bool /*bAllowProjectOpening*/, bool /*bAllowProjectCreate*/);
-FUNC_DECLARE_DELEGATE(FCreateProjectBrowser, TSharedRef<SWidget>, bool /*bInShowHeaderBar*/, bool /*bInAllowScrolling*/);
+FUNC_DECLARE_DELEGATE(FCreateProjectBrowser, TSharedRef<SWidget>, bool /*bOnlyShowList*/, bool /*bAllowScrolling*/);
 
 /**
  * Interface for main frame modules.
@@ -250,7 +250,7 @@ public:
 	/**
 	 * Creates the project browser for embedding.
 	 */
-	virtual TSharedRef<SWidget> CreateProjectBrowser(bool bInShowHeaderBar = true, bool bInAllowScrolling = true) = 0;
+	virtual TSharedRef<SWidget> CreateProjectBrowser(bool bOnlyShowList = true, bool bAllowScrolling = true) = 0;
 
 	/**
 	 * Overrides the functionality of ShowProjectBrowser with the provided delegate.
