[HomeScreen] Update urls

diff --git a/Engine/Source/Editor/MainFrame/Private/MainFrameModule.cpp b/Engine/Source/Editor/MainFrame/Private/MainFrameModule.cpp
--- a/Engine/Source/Editor/MainFrame/Private/MainFrameModule.cpp
+++ b/Engine/Source/Editor/MainFrame/Private/MainFrameModule.cpp
@@ -1102,7 +1102,7 @@ TSharedRef<SWidget> FMainFrameModule::CreateHomeScreenWidget(const TSharedPtr<SD
 		IWebBrowserSingleton* WebBrowserSingleton = WebBrowserModule.GetSingleton();
 
 		FCreateBrowserWindowSettings WindowSettings;
-		WindowSettings.InitialURL = TEXT("https://editor.unrealengine.com/en-US/recent-projects");
+		WindowSettings.InitialURL = TEXT("about:blank");
 		WindowSettings.bUseTransparency = true;
 		WindowSettings.BrowserFrameRate = 60;
 
