Added the ability to override whether to always show the recent project list and what the recent project list should look like to the Home Screen settings.

#rnx

diff --git a/Engine/Source/Editor/MainFrame/Private/HomeScreen/SHomeScreen.cpp b/Engine/Source/Editor/MainFrame/Private/HomeScreen/SHomeScreen.cpp
--- a/Engine/Source/Editor/MainFrame/Private/HomeScreen/SHomeScreen.cpp
+++ b/Engine/Source/Editor/MainFrame/Private/HomeScreen/SHomeScreen.cpp
@@ -125,7 +125,7 @@ void SHomeScreen::Construct(const FArguments& InArgs, const TSharedPtr<SWebBrows
 		WebObject->OnTutorialProjectRequested().AddSP(this, &SHomeScreen::OnOpenGettingStartedProject);
 	}
 
-	bHasLatestEngineProjects = HasAlreadyLatestEngineProject();
+	bHasLatestEngineProjects = InArgs._AlwaysShowRecentList || HasAlreadyLatestEngineProject();
 	MainHomeSelection = bHasLatestEngineProjects ? EMainSectionMenu::Home : EMainSectionMenu::GettingStarted;
 
 	// Fixed Button Size
@@ -150,11 +150,6 @@ void SHomeScreen::Construct(const FArguments& InArgs, const TSharedPtr<SWebBrows
 	CreateSocialMediaButtons(TwitchButton, TEXT("https://www.twitch.tv/unrealengine"), FAppStyle::GetBrush("SocialMedia.Twitch"));
 	CreateSocialMediaButtons(InstagramButton, TEXT("https://www.instagram.com/unrealengine/"), FAppStyle::GetBrush("SocialMedia.Instagram"));
 
-	// Project Browser
-	constexpr bool bOnlyShowList = true;
-	constexpr bool bProjectBrowserAllowScrolling = false;
-	ProjectBrowser = IMainFrameModule::Get().CreateProjectBrowser(bOnlyShowList, bProjectBrowserAllowScrolling);
-
 	// Check for internet connection on startup
 	CheckInternetConnection();
 
@@ -551,12 +546,7 @@ void SHomeScreen::Construct(const FArguments& InArgs, const TSharedPtr<SWebBrows
 
 									+ SVerticalBox::Slot()
 									[
-										SNew(SBox)
-										.HeightOverride(172.f)
-										.Padding(0.f)
-										[
-											ProjectBrowser.ToSharedRef()
-										]
+										InArgs._RecentProjectList
 									]
 								]
 							]
diff --git a/Engine/Source/Editor/MainFrame/Private/HomeScreen/SHomeScreen.h b/Engine/Source/Editor/MainFrame/Private/HomeScreen/SHomeScreen.h
--- a/Engine/Source/Editor/MainFrame/Private/HomeScreen/SHomeScreen.h
+++ b/Engine/Source/Editor/MainFrame/Private/HomeScreen/SHomeScreen.h
@@ -1,4 +1,4 @@
-﻿// Copyright Epic Games, Inc. All Rights Reserved.
+// Copyright Epic Games, Inc. All Rights Reserved.
 
 #pragma once
 #include "Engine/TimerHandle.h"
@@ -9,6 +9,8 @@
 #include "Styling/SlateColor.h"
 #include "Styling/SlateTypes.h"
 #include "Widgets/SCompoundWidget.h"
+#include "Widgets/SNullWidget.h"
+#include "Widgets/SWidget.h"
 
 class SButton;
 class SCheckBox;
@@ -23,9 +25,13 @@ public:
 	SLATE_BEGIN_ARGS(SHomeScreen)
 		: _ResourceCommunityEntries()
 		, _SectionLinks()
+		, _AlwaysShowRecentList(false)
+		, _RecentProjectList(SNullWidget::NullWidget)
 	{}
 		SLATE_ARGUMENT(TArray<FHomeScreenResourceCommunityEntry>, ResourceCommunityEntries)
 		SLATE_ARGUMENT(FHomeScreenSectionToLinkMap, SectionLinks)
+		SLATE_ARGUMENT(bool, AlwaysShowRecentList)
+		SLATE_ARGUMENT(TSharedRef<SWidget>, RecentProjectList)
 	SLATE_END_ARGS()
 
 	SHomeScreen();
@@ -125,9 +131,6 @@ private:
 	/** Combo box for autoload project selection */
 	TSharedPtr<SComboButton> AutoLoadProjectComboBox;
 
-	/** Widget for project browser */
-	TSharedPtr<SWidget> ProjectBrowser;
-
 	/** Current combo button menu */
 	TWeakPtr<SWidget> ComboButtonMenuWeak;
 
diff --git a/Engine/Source/Editor/MainFrame/Private/MainFrameModule.cpp b/Engine/Source/Editor/MainFrame/Private/MainFrameModule.cpp
--- a/Engine/Source/Editor/MainFrame/Private/MainFrameModule.cpp
+++ b/Engine/Source/Editor/MainFrame/Private/MainFrameModule.cpp
@@ -154,12 +154,26 @@ namespace HomeScreenHelper
 		return OutSectionLinks;
 	}
 
+	TSharedRef<SWidget> GetRecentProjectListWidget()
+	{
+		constexpr bool bOnlyShowList = true;
+		constexpr bool bProjectBrowserAllowScrolling = false;
+		return SNew(SBox)
+			.HeightOverride(172.f)
+			.Padding(0.f)
+			[
+				IMainFrameModule::Get().CreateProjectBrowser(bOnlyShowList, bProjectBrowserAllowScrolling)
+			];
+	}
+
 	// Get the default UE home screen settings
 	FHomeScreenWidgetSettings GetHomeScreenSettings()
 	{
 		FHomeScreenWidgetSettings HomeScreenWidgetSettings;
 		HomeScreenWidgetSettings.ResourcesEntries = GetUEResourceEntries();
 		HomeScreenWidgetSettings.SectionLinks = GetUESectionLinks();
+		HomeScreenWidgetSettings.bAlwaysShowRecentList = false;
+		HomeScreenWidgetSettings.RecentProjectListWidget = GetRecentProjectListWidget();
 		return HomeScreenWidgetSettings;
 	}
 }
@@ -1159,7 +1173,9 @@ TSharedRef<SWidget> FMainFrameModule::CreateHomeScreenWidget(const TSharedPtr<SD
 
 	return SNew(SHomeScreen, WebBrowser, TWeakObjectPtr(WebObject.Get()))
 		.SectionLinks(InSettings.SectionLinks)
-		.ResourceCommunityEntries(InSettings.ResourcesEntries);
+		.ResourceCommunityEntries(InSettings.ResourcesEntries)
+		.AlwaysShowRecentList(InSettings.bAlwaysShowRecentList)
+		.RecentProjectList(InSettings.RecentProjectListWidget);
 }
 
 void FMainFrameModule::HandleCodeAccessorOpenFileFailed(const FString& Filename)
diff --git a/Engine/Source/Editor/MainFrame/Public/Settings/HomeScreenCommon.h b/Engine/Source/Editor/MainFrame/Public/Settings/HomeScreenCommon.h
--- a/Engine/Source/Editor/MainFrame/Public/Settings/HomeScreenCommon.h
+++ b/Engine/Source/Editor/MainFrame/Public/Settings/HomeScreenCommon.h
@@ -1,4 +1,4 @@
-﻿// Copyright Epic Games, Inc. All Rights Reserved.
+// Copyright Epic Games, Inc. All Rights Reserved.
 
 #pragma once
 #include "Containers/Array.h"
@@ -6,6 +6,8 @@
 #include "Containers/UnrealString.h"
 #include "Internationalization/Text.h"
 #include "UObject/NameTypes.h"
+#include "Widgets/SNullWidget.h"
+#include "Widgets/SWidget.h"
 #include "HomeScreenCommon.generated.h"
 
 UENUM()
@@ -51,6 +53,12 @@ struct FHomeScreenResourceCommunityEntry
 /** Setting used to create the HomeScreen, if leaved empty the defaults are used instead */
 struct FHomeScreenWidgetSettings
 {
+	FHomeScreenWidgetSettings()
+		: bAlwaysShowRecentList(false)
+		, RecentProjectListWidget(SNullWidget::NullWidget)
+	{
+	}
+
 	/** 
 	 * Section x Link connection, if leaved empty it will use the UE default section x links.
 	 * If at least 1 is populated it will use this instead, every missing section are skipped from being created in this case
@@ -61,4 +69,14 @@ struct FHomeScreenWidgetSettings
 	 * Additional Resource entry to add to the default ones
 	 */
 	TArray<FHomeScreenResourceCommunityEntry> ResourcesEntries;
+
+	/**
+	 * Whether the recent project list should always be displayed without checking for their presence.
+	 */
+	bool bAlwaysShowRecentList;
+
+	/**
+	 * Widget displaying the recent project list.
+	 */
+	TSharedRef<SWidget> RecentProjectListWidget;
 };
