HomeScreen: Added the ability to override the Unreal Logo

#rnx

diff --git a/Engine/Source/Editor/MainFrame/Private/HomeScreen/SHomeScreen.cpp b/Engine/Source/Editor/MainFrame/Private/HomeScreen/SHomeScreen.cpp
--- a/Engine/Source/Editor/MainFrame/Private/HomeScreen/SHomeScreen.cpp
+++ b/Engine/Source/Editor/MainFrame/Private/HomeScreen/SHomeScreen.cpp
@@ -179,7 +179,7 @@ void SHomeScreen::Construct(const FArguments& InArgs, const TSharedPtr<SWebBrows
 					.Padding(0.f, 32.f)
 					[
 						SNew(SImage)
-						.Image(FAppStyle::GetBrush("HomeScreen.UnrealLogo"))
+						.Image(InArgs._UnrealLogoBrush)
 					]
 
 					// Create and MyFolder Buttons
diff --git a/Engine/Source/Editor/MainFrame/Private/HomeScreen/SHomeScreen.h b/Engine/Source/Editor/MainFrame/Private/HomeScreen/SHomeScreen.h
--- a/Engine/Source/Editor/MainFrame/Private/HomeScreen/SHomeScreen.h
+++ b/Engine/Source/Editor/MainFrame/Private/HomeScreen/SHomeScreen.h
@@ -27,11 +27,13 @@ public:
 		, _SectionLinks()
 		, _AlwaysShowRecentList(false)
 		, _RecentProjectList(SNullWidget::NullWidget)
+		, _UnrealLogoBrush(nullptr)
 	{}
 		SLATE_ARGUMENT(TArray<FHomeScreenResourceCommunityEntry>, ResourceCommunityEntries)
 		SLATE_ARGUMENT(FHomeScreenSectionToLinkMap, SectionLinks)
 		SLATE_ARGUMENT(bool, AlwaysShowRecentList)
 		SLATE_ARGUMENT(TSharedRef<SWidget>, RecentProjectList)
+		SLATE_ARGUMENT(const FSlateBrush*, UnrealLogoBrush)
 	SLATE_END_ARGS()
 
 	SHomeScreen();
diff --git a/Engine/Source/Editor/MainFrame/Private/MainFrameModule.cpp b/Engine/Source/Editor/MainFrame/Private/MainFrameModule.cpp
--- a/Engine/Source/Editor/MainFrame/Private/MainFrameModule.cpp
+++ b/Engine/Source/Editor/MainFrame/Private/MainFrameModule.cpp
@@ -174,6 +174,7 @@ namespace HomeScreenHelper
 		HomeScreenWidgetSettings.SectionLinks = GetUESectionLinks();
 		HomeScreenWidgetSettings.bAlwaysShowRecentList = false;
 		HomeScreenWidgetSettings.RecentProjectListWidget = GetRecentProjectListWidget();
+		HomeScreenWidgetSettings.UnrealLogoBrush = FAppStyle::GetBrush("HomeScreen.UnrealLogo");
 		return HomeScreenWidgetSettings;
 	}
 }
@@ -1175,7 +1176,8 @@ TSharedRef<SWidget> FMainFrameModule::CreateHomeScreenWidget(const TSharedPtr<SD
 		.SectionLinks(InSettings.SectionLinks)
 		.ResourceCommunityEntries(InSettings.ResourcesEntries)
 		.AlwaysShowRecentList(InSettings.bAlwaysShowRecentList)
-		.RecentProjectList(InSettings.RecentProjectListWidget);
+		.RecentProjectList(InSettings.RecentProjectListWidget)
+		.UnrealLogoBrush(InSettings.UnrealLogoBrush);
 }
 
 void FMainFrameModule::HandleCodeAccessorOpenFileFailed(const FString& Filename)
diff --git a/Engine/Source/Editor/MainFrame/Public/Settings/HomeScreenCommon.h b/Engine/Source/Editor/MainFrame/Public/Settings/HomeScreenCommon.h
--- a/Engine/Source/Editor/MainFrame/Public/Settings/HomeScreenCommon.h
+++ b/Engine/Source/Editor/MainFrame/Public/Settings/HomeScreenCommon.h
@@ -56,6 +56,7 @@ struct FHomeScreenWidgetSettings
 	FHomeScreenWidgetSettings()
 		: bAlwaysShowRecentList(false)
 		, RecentProjectListWidget(SNullWidget::NullWidget)
+		, UnrealLogoBrush(nullptr)
 	{
 	}
 
@@ -79,4 +80,9 @@ struct FHomeScreenWidgetSettings
 	 * Widget displaying the recent project list.
 	 */
 	TSharedRef<SWidget> RecentProjectListWidget;
+
+	/**
+	 * Brush representing the icon for the Unreal Logo.
+	 */
+	const FSlateBrush* UnrealLogoBrush;
 };
