[HomeScreen] User home screen startup issue with the Project Browser

#Synced-CL 48083090

diff --git a/Engine/Source/Editor/MainFrame/Private/MainFrameModule.cpp b/Engine/Source/Editor/MainFrame/Private/MainFrameModule.cpp
--- a/Engine/Source/Editor/MainFrame/Private/MainFrameModule.cpp
+++ b/Engine/Source/Editor/MainFrame/Private/MainFrameModule.cpp
@@ -294,11 +294,42 @@ void FMainFrameModule::RecreateDefaultMainFrame(const bool bStartImmersive, cons
 	CreateDefaultMainFrameAuxiliary(bStartImmersive, bStartPIE, /*bIsBeingRecreated*/true);
 }
 
-void FMainFrameModule::CreateDefaultMainFrameAuxiliary(const bool bStartImmersive, const bool bStartPIE, const bool bIsBeingRecreated)
+void FMainFrameModule::ReconstructDefaultMainFrameContent(const bool bStartImmersive, const bool bStartPIE)
+{
+	// Clean previous default main frame
+	if (IsWindowInitialized())
+	{
+		// Clean FGlobalTabmanager
+		FGlobalTabmanager::Get()->CloseAllAreas();
+
+		// ::RequestDestroyWindow will remove the window from the TArray so we
+		// should iterate over it backwards to make sure that the WindowIndex is still correct for the TopLevelWindows.
+		TSharedRef<SWindow> RootWindow = GetParentWindow().ToSharedRef();
+
+		// Iterate every TopLevelWindows and destroy them if they are not the MainFrame ParentWindow
+		TArray<TSharedRef<SWindow>> TopLevelWindows = FSlateApplication::Get().GetTopLevelWindows();
+		for (int32 TopLevelWindowIndex = TopLevelWindows.Num() -1; TopLevelWindowIndex >= 0; --TopLevelWindowIndex)
+		{
+			if (TopLevelWindows[TopLevelWindowIndex] == RootWindow)
+			{
+				continue;
+			}
+
+			FSlateApplication::Get().RequestDestroyWindow(TopLevelWindows[TopLevelWindowIndex]);
+		}
+
+		FSlateApplication::Get().DestroyWindowsImmediately();
+		
+		// (Re-)populate the default main frame
+		CreateDefaultMainFrameAuxiliary(bStartImmersive, bStartPIE, false, true);
+	}
+}
+
+void FMainFrameModule::CreateDefaultMainFrameAuxiliary(const bool bStartImmersive, const bool bStartPIE, const bool bIsBeingRecreated, const bool bIsJustCreatingContent)
 {
 	TRACE_CPUPROFILER_EVENT_SCOPE(FMainFrameModule::CreateDefaultMainFrame);
 
-	if (!IsWindowInitialized())
+	if (!IsWindowInitialized() || bIsJustCreatingContent)
 	{
 		FRootWindowLocation DefaultWindowLocation;
 
@@ -347,25 +378,73 @@ void FMainFrameModule::CreateDefaultMainFrameAuxiliary(const bool bStartImmersiv
 			WindowConfig.WindowTitle = GetApplicationTitle( bIncludeGameName );
 		}
 
-		TSharedRef<SWindow> RootWindow = SNew(SWindow)
-			.AutoCenter(WindowConfig.CenterRules)
-			.Title( WindowConfig.WindowTitle )
-			.CloseButtonToolTipText(LOCTEXT("MainFrameModuleCloseButtonToolTip", "Close Unreal Editor"))
-			.IsInitiallyMaximized( DefaultWindowLocation.InitiallyMaximized )
-			.ScreenPosition( DefaultWindowLocation.ScreenPosition )
-			.ClientSize( DefaultWindowLocation.WindowSize )
-			.CreateTitleBar( !WindowConfig.bEmbedTitleAreaContent )
-			.SizingRule( WindowConfig.bIsUserSizable ? ESizingRule::UserSized : ESizingRule::FixedSize )
-			.SupportsMaximize( WindowConfig.bSupportsMaximize )
-			.SupportsMinimize( WindowConfig.bSupportsMinimize );
-
-		const bool bShowRootWindowImmediately = false;
-		FSlateApplication::Get().AddWindow( RootWindow, bShowRootWindowImmediately );
-
-		FGlobalTabmanager::Get()->SetRootWindow(RootWindow);
-		FGlobalTabmanager::Get()->SetAllowWindowMenuBar(true);
-
-		FSlateNotificationManager::Get().SetRootWindow(RootWindow);
+		const FText CloseButtonToolTipText = LOCTEXT("MainFrameModuleCloseButtonToolTip", "Close Unreal Editor");
+
+		if (!bIsJustCreatingContent)
+		{
+			TSharedRef<SWindow> RootWindow = SNew(SWindow)
+				.AutoCenter(WindowConfig.CenterRules)
+				.Title(WindowConfig.WindowTitle)
+				.CloseButtonToolTipText(CloseButtonToolTipText)
+				.IsInitiallyMaximized(DefaultWindowLocation.InitiallyMaximized)
+				.ScreenPosition(DefaultWindowLocation.ScreenPosition)
+				.ClientSize(DefaultWindowLocation.WindowSize)
+				.CreateTitleBar(!WindowConfig.bEmbedTitleAreaContent)
+				.SizingRule(WindowConfig.bIsUserSizable ? ESizingRule::UserSized : ESizingRule::FixedSize)
+				.SupportsMaximize(WindowConfig.bSupportsMaximize)
+				.SupportsMinimize(WindowConfig.bSupportsMinimize);
+
+			const bool bShowRootWindowImmediately = false;
+			FSlateApplication::Get().AddWindow( RootWindow, bShowRootWindowImmediately );
+
+			FGlobalTabmanager::Get()->SetRootWindow(RootWindow);
+			FGlobalTabmanager::Get()->SetAllowWindowMenuBar(true);
+
+			FSlateNotificationManager::Get().SetRootWindow(RootWindow);
+		}
+
+		// If we are forcing the ContentCreation to be able to reach here the GetParentWindow should be valid
+		TSharedPtr<SWindow> RootWindowPtr = bIsJustCreatingContent ? GetParentWindow() : FGlobalTabmanager::Get()->GetRootWindow();
+		
+		// The window has to be valid
+		check(RootWindowPtr.IsValid())
+
+		TSharedRef<SWindow> RootWindow = RootWindowPtr.ToSharedRef();
+
+		if (bIsJustCreatingContent)
+		{
+			if (DefaultWindowLocation.InitiallyMaximized)
+			{
+				// if initially maximized get the fullscreen size and use that for the size to make the Maximize jump less visible
+				// TODO: Should be looking at maximizing without showing the window or to make the InitiallyMaximized work on second window constructs
+				DefaultWindowLocation.WindowSize = RootWindow->GetFullScreenInfo().GetSize();
+			}
+			else if (DefaultWindowLocation.WindowSize.IsNearlyZero())
+			{
+				// Resize the window with the previous size if not maximized and if the window size was not set (or nearly 0)
+				UE::Slate::FDeprecateVector2DResult SizeToRestore = RootWindow->GetSizeInScreen();
+				DefaultWindowLocation.WindowSize = SizeToRestore;
+			}
+
+			SWindow::FArguments Args;
+			Args.AutoCenter(WindowConfig.CenterRules)
+				.Title(WindowConfig.WindowTitle)
+				.CloseButtonToolTipText(CloseButtonToolTipText)
+				.IsInitiallyMaximized(DefaultWindowLocation.InitiallyMaximized)
+				.ScreenPosition( DefaultWindowLocation.ScreenPosition)
+				.ClientSize(DefaultWindowLocation.WindowSize)
+				.CreateTitleBar(!WindowConfig.bEmbedTitleAreaContent)
+				.SizingRule(WindowConfig.bIsUserSizable ? ESizingRule::UserSized : ESizingRule::FixedSize)
+				.SupportsMaximize(WindowConfig.bSupportsMaximize)
+				.SupportsMinimize(WindowConfig.bSupportsMinimize);
+
+			RootWindow->Construct(Args);
+			// InitiallyMaximized is not correctly used if the window was already created and visible, so do it manually if requested
+			if (DefaultWindowLocation.InitiallyMaximized)
+			{
+				RootWindow->Maximize();
+			}
+		}
 
 		bool bLevelEditorIsMainTab = false;
 		if (!bShowStartupDialogInPlaceOfMainEditor)
diff --git a/Engine/Source/Editor/MainFrame/Private/MainFrameModule.h b/Engine/Source/Editor/MainFrame/Private/MainFrameModule.h
--- a/Engine/Source/Editor/MainFrame/Private/MainFrameModule.h
+++ b/Engine/Source/Editor/MainFrame/Private/MainFrameModule.h
@@ -64,12 +64,14 @@ public:
 
 	virtual void CreateDefaultMainFrame(const bool bStartImmersive, const bool bStartPIE) override;
 	virtual void RecreateDefaultMainFrame(const bool bStartImmersive, const bool bStartPIE) override;
+	virtual void ReconstructDefaultMainFrameContent(const bool bStartImmersive, const bool bStartPIE) override;
 private:
 	/**
 	 * Shared code between CreateDefaultMainFrame and RecreateDefaultMainFrame
 	 * @param bIsBeingRecreated False if it is being called by first time (CreateDefaultMainFrame), and true if it is being recreated (RecreateDefaultMainFrame). If recreated, it will also display 
+	 * @param bIsJustCreatingContent False if it is being called by first time (CreateDefaultMainFrame), and true if the content is being reconstructed (ReconstructDefaultMainFrameContent).
 	 */
-	virtual void CreateDefaultMainFrameAuxiliary(const bool bStartImmersive, const bool bStartPIE, const bool bIsBeingRecreated);
+	virtual void CreateDefaultMainFrameAuxiliary(const bool bStartImmersive, const bool bStartPIE, const bool bIsBeingRecreated, const bool bIsJustCreatingContent = false);
 
 public:
 	virtual bool IsRecreatingDefaultMainFrame() const override;
diff --git a/Engine/Source/Editor/MainFrame/Public/Interfaces/IMainFrameModule.h b/Engine/Source/Editor/MainFrame/Public/Interfaces/IMainFrameModule.h
--- a/Engine/Source/Editor/MainFrame/Public/Interfaces/IMainFrameModule.h
+++ b/Engine/Source/Editor/MainFrame/Public/Interfaces/IMainFrameModule.h
@@ -66,6 +66,14 @@ public:
 	 */
 	virtual void RecreateDefaultMainFrame(const bool bStartImmersive, const bool bStartPIE) = 0;
 
+	/**
+	 * Recreates the default editor main frame content without destroying the old one.
+	 *
+	 * @param bStartImmersive True to force a main frame viewport into immersive mode
+	 * @param bStartPIE True to start a PIE session right away
+	 */
+	virtual void ReconstructDefaultMainFrameContent(const bool bStartImmersive, const bool bStartPIE) = 0;
+
 	/**
 	 * Returns true if the Default Main Frame is being recreated. 
 	 */
