[Home Screen] Add dynamic hook for the startup setting
Adding API to allow dynamic setting to be bound as the most recent project startup if a different setting should be used

#Synced-CL 48654334

diff --git a/Engine/Source/Editor/MainFrame/Private/HomeScreen/HomeScreenSettings.cpp b/Engine/Source/Editor/MainFrame/Private/HomeScreen/HomeScreenSettings.cpp
--- a/Engine/Source/Editor/MainFrame/Private/HomeScreen/HomeScreenSettings.cpp
+++ b/Engine/Source/Editor/MainFrame/Private/HomeScreen/HomeScreenSettings.cpp
@@ -4,6 +4,13 @@
 #include "Settings/HomeScreenCommon.h"
 #include "Settings/EditorSettings.h"
 
+UHomeScreenSettings::UHomeScreenSettings()
+{
+	FOnSetLoadAtStartup OnSetLoadAtStartup = FOnSetLoadAtStartup::CreateUObject(this, &UHomeScreenSettings::SetLoadAtStartupChanged);
+	FOnGetLoadAtStartup OnGetLoadAtStartup = FOnGetLoadAtStartup::CreateUObject(this, &UHomeScreenSettings::GetLoadAtStartupChanged);
+	HomeScreenLoadAtStartupSetting = FHomeScreenLoadAtStartupSetting(OnSetLoadAtStartup, OnGetLoadAtStartup);
+}
+
 void UHomeScreenSettings::PostEditChangeProperty(FPropertyChangedEvent& PropertyChangedEvent)
 {
 	Super::PostEditChangeProperty(PropertyChangedEvent);
@@ -39,10 +46,44 @@ void UHomeScreenSettings::PostInitProperties()
 	}
 }
 
+FHomeScreenLoadAtStartupSetting& UHomeScreenSettings::GetLoadAtStartupSetting()
+{
+	return HomeScreenLoadAtStartupSetting;
+}
+
 void UHomeScreenSettings::OnMostRecentProjectEditorSettingChanged(bool bInAutoLoad)
 {
 	LoadAtStartup = bInAutoLoad ? EAutoLoadProject::LastProject : EAutoLoadProject::HomeScreen;
 
 	// Broadcast that the setting changed
+	HomeScreenLoadAtStartupSetting.OnLoadAtStartupChanged().Broadcast(LoadAtStartup);
+
+	PRAGMA_DISABLE_DEPRECATION_WARNINGS
 	OnLoadAtStartupChanged().Broadcast(LoadAtStartup);
+	PRAGMA_ENABLE_DEPRECATION_WARNINGS
+}
+
+void UHomeScreenSettings::SetLoadAtStartupChanged(bool bInAutoLoad)
+{
+	EAutoLoadProject NewValue = bInAutoLoad ? EAutoLoadProject::LastProject : EAutoLoadProject::HomeScreen;
+
+	if (LoadAtStartup == NewValue)
+	{
+		return;
+	}
+
+	LoadAtStartup = NewValue;
+
+	// Manually trigger the property changed event so that the home screen editor settings UI updates to reflect the change
+	FProperty* LoadAtStartupProperty = FindFProperty<FProperty>(GetClass(), "LoadAtStartup");
+	if (LoadAtStartupProperty != NULL)
+	{
+		FPropertyChangedEvent PropertyUpdateStruct(LoadAtStartupProperty);
+		PostEditChangeProperty(PropertyUpdateStruct);
+	}
+}
+
+bool UHomeScreenSettings::GetLoadAtStartupChanged() const
+{
+	return LoadAtStartup == EAutoLoadProject::LastProject;
 }
diff --git a/Engine/Source/Editor/MainFrame/Private/HomeScreen/SHomeScreen.cpp b/Engine/Source/Editor/MainFrame/Private/HomeScreen/SHomeScreen.cpp
--- a/Engine/Source/Editor/MainFrame/Private/HomeScreen/SHomeScreen.cpp
+++ b/Engine/Source/Editor/MainFrame/Private/HomeScreen/SHomeScreen.cpp
@@ -81,11 +81,6 @@ SHomeScreen::SHomeScreen()
 
 SHomeScreen::~SHomeScreen()
 {
-	if (UHomeScreenSettings* HomeScreenSettings = GetMutableDefault<UHomeScreenSettings>())
-	{
-		HomeScreenSettings->OnLoadAtStartupChanged().RemoveAll(this);
-	}
-
 	if (GEditor && GEditor->IsTimerManagerValid())
 	{
 		GEditor->GetTimerManager()->ClearTimer(CheckInternetConnectionTimerHandle);
@@ -111,12 +106,14 @@ void SHomeScreen::Construct(const FArguments& InArgs, const TSharedPtr<SWebBrows
 	WebBrowser = InWebBrowser;
 	WebObject = InWebObject;
 	SectionLinks = InArgs._SectionLinks;
+	OnGetStartupSettingDelegate = InArgs._OnGetStartupSetting;
 
 	AutoLoadProjectComboBoxSelection = EAutoLoadProject::HomeScreen;
-	if (UHomeScreenSettings* HomeScreenSettings = GetMutableDefault<UHomeScreenSettings>())
+	if (OnGetStartupSettingDelegate.IsBound())
 	{
-		HomeScreenSettings->OnLoadAtStartupChanged().AddSP(this, &SHomeScreen::OnLoadAtStartupSettingChanged);
-		AutoLoadProjectComboBoxSelection = HomeScreenSettings->LoadAtStartup;
+		FHomeScreenLoadAtStartupSetting& LoadAtStartupSetting = OnGetStartupSettingDelegate.Execute();
+		LoadAtStartupSetting.OnLoadAtStartupChanged().AddSP(this, &SHomeScreen::OnLoadAtStartupSettingChanged);
+		AutoLoadProjectComboBoxSelection = LoadAtStartupSetting.GetOnLoadAtStartup() ? EAutoLoadProject::LastProject : EAutoLoadProject::HomeScreen;
 	}
 
 	if (WebObject.IsValid())
@@ -393,6 +390,7 @@ void SHomeScreen::Construct(const FArguments& InArgs, const TSharedPtr<SWebBrows
 						.HAlign(HAlign_Left)
 						[
 							SNew(STextBlock)
+							.Visibility(OnGetStartupSettingDelegate.IsBound() ? EVisibility::Visible : EVisibility::Collapsed)
 							.Font(FAppStyle::GetFontStyle("HomeScreen.Font.DefaultSize"))
 							.Text(LOCTEXT("ProjectHomeLoad", "Load on Startup"))
 						]
@@ -402,6 +400,7 @@ void SHomeScreen::Construct(const FArguments& InArgs, const TSharedPtr<SWebBrows
 						.Padding(16.f, 8.f, 16.f, 36.f)
 						[
 							SAssignNew(AutoLoadProjectComboBox, SComboButton)
+							.Visibility(OnGetStartupSettingDelegate.IsBound() ? EVisibility::Visible : EVisibility::Collapsed)
 							.ComboButtonStyle(&FAppStyle::Get().GetWidgetStyle<FComboButtonStyle>("HomeScreen.ComboButton"))
 							.MenuPlacement(EMenuPlacement::MenuPlacement_CenteredAboveAnchor)
 							.HasDownArrow(true)
@@ -650,18 +649,13 @@ FReply SHomeScreen::OnAutoLoadOptionChanged(EAutoLoadProject InAutoLoadOption)
 		return FReply::Handled();
 	}
 
-	// Do not change the option if we can't change the setting otherwise they are not synced, the PostEditChangeProperty will handle the call to set our ComboBoxSelection.
-	if (UHomeScreenSettings* HomeScreenSettings = GetMutableDefault<UHomeScreenSettings>())
+	// SetOnLoad should change the setting value and the setting class should call the OnChanged that will call out OnLoadAtStartupSettingChanged
+	// This is done so that the setting class will decide to change the setting or not since there may be more logic there
+	// If the Set is later on skipped for example we will still have the setting here and in the setting class synced
+	if (OnGetStartupSettingDelegate.IsBound())
 	{
-		HomeScreenSettings->LoadAtStartup = InAutoLoadOption;
-
-		// Manually trigger the property changed event so that the home screen editor settings UI updates to reflect the change
-		FProperty* LoadAtStartupProperty = FindFProperty<FProperty>(HomeScreenSettings->GetClass(), "LoadAtStartup");
-		if (LoadAtStartupProperty != NULL)
-		{
-			FPropertyChangedEvent PropertyUpdateStruct(LoadAtStartupProperty);
-			HomeScreenSettings->PostEditChangeProperty(PropertyUpdateStruct);
-		}
+		const FHomeScreenLoadAtStartupSetting& LoadAtStartupSetting = OnGetStartupSettingDelegate.Execute();
+		LoadAtStartupSetting.SetOnLoadAtStartup(InAutoLoadOption == EAutoLoadProject::LastProject);
 	}
 
 	if (TSharedPtr<SWidget> ComboButtonMenu = ComboButtonMenuWeak.Pin())
diff --git a/Engine/Source/Editor/MainFrame/Private/HomeScreen/SHomeScreen.h b/Engine/Source/Editor/MainFrame/Private/HomeScreen/SHomeScreen.h
--- a/Engine/Source/Editor/MainFrame/Private/HomeScreen/SHomeScreen.h
+++ b/Engine/Source/Editor/MainFrame/Private/HomeScreen/SHomeScreen.h
@@ -28,12 +28,14 @@ public:
 		, _AlwaysShowRecentList(false)
 		, _RecentProjectList(SNullWidget::NullWidget)
 		, _UnrealLogoBrush(nullptr)
+		, _OnGetStartupSetting()
 	{}
 		SLATE_ARGUMENT(TArray<FHomeScreenResourceCommunityEntry>, ResourceCommunityEntries)
 		SLATE_ARGUMENT(FHomeScreenSectionToLinkMap, SectionLinks)
 		SLATE_ARGUMENT(bool, AlwaysShowRecentList)
 		SLATE_ARGUMENT(TSharedRef<SWidget>, RecentProjectList)
 		SLATE_ARGUMENT(const FSlateBrush*, UnrealLogoBrush)
+		SLATE_ARGUMENT(FOnGetStartupSetting, OnGetStartupSetting)
 	SLATE_END_ARGS()
 
 	SHomeScreen();
@@ -171,7 +173,10 @@ private:
 	
 	/** WebBrowser object passed down from the MainFrameModule */
 	TWeakObjectPtr<UHomeScreenWeb> WebObject;
-	
+
+	/** HomeScreen start up setting getter */
+	FOnGetStartupSetting OnGetStartupSettingDelegate;
+
 	/** Section x link combination */
 	FHomeScreenSectionToLinkMap SectionLinks;
 };
diff --git a/Engine/Source/Editor/MainFrame/Private/MainFrameModule.cpp b/Engine/Source/Editor/MainFrame/Private/MainFrameModule.cpp
--- a/Engine/Source/Editor/MainFrame/Private/MainFrameModule.cpp
+++ b/Engine/Source/Editor/MainFrame/Private/MainFrameModule.cpp
@@ -46,6 +46,7 @@
 #include "SWebBrowser.h"
 #include "ToolMenus.h"
 #include "WebBrowserModule.h"
+#include "Settings/HomeScreenSettings.h"
 #include "Widgets/Docking/SDockTab.h"
 
 DEFINE_LOG_CATEGORY(LogMainFrame);
@@ -175,6 +176,7 @@ namespace HomeScreenHelper
 		HomeScreenWidgetSettings.bAlwaysShowRecentList = false;
 		HomeScreenWidgetSettings.RecentProjectListWidget = GetRecentProjectListWidget();
 		HomeScreenWidgetSettings.UnrealLogoBrush = FAppStyle::GetBrush("HomeScreen.UnrealLogo");
+		HomeScreenWidgetSettings.OnGetStartupSetting.BindUObject(GetMutableDefault<UHomeScreenSettings>(), &UHomeScreenSettings::GetLoadAtStartupSetting);
 		return HomeScreenWidgetSettings;
 	}
 }
@@ -1252,7 +1254,8 @@ TSharedRef<SWidget> FMainFrameModule::CreateHomeScreenWidget(const TSharedPtr<SD
 		.ResourceCommunityEntries(InSettings.ResourcesEntries)
 		.AlwaysShowRecentList(InSettings.bAlwaysShowRecentList)
 		.RecentProjectList(InSettings.RecentProjectListWidget)
-		.UnrealLogoBrush(InSettings.UnrealLogoBrush);
+		.UnrealLogoBrush(InSettings.UnrealLogoBrush)
+		.OnGetStartupSetting(InSettings.OnGetStartupSetting);
 }
 
 void FMainFrameModule::HandleCodeAccessorOpenFileFailed(const FString& Filename)
diff --git a/Engine/Source/Editor/MainFrame/Public/Settings/HomeScreenCommon.h b/Engine/Source/Editor/MainFrame/Public/Settings/HomeScreenCommon.h
--- a/Engine/Source/Editor/MainFrame/Public/Settings/HomeScreenCommon.h
+++ b/Engine/Source/Editor/MainFrame/Public/Settings/HomeScreenCommon.h
@@ -28,6 +28,10 @@ enum class EAutoLoadProject
 	MAX UMETA(Hidden)
 };
 
+DECLARE_MULTICAST_DELEGATE_OneParam(FOnLoadAtStartupChanged, EAutoLoadProject)
+DECLARE_DELEGATE_OneParam(FOnSetLoadAtStartup, bool)
+DECLARE_DELEGATE_RetVal(bool, FOnGetLoadAtStartup)
+
 /**
  * Struct containing information on the MainSections in the HomeScreen
  */
@@ -50,6 +54,29 @@ struct FHomeScreenResourceCommunityEntry
 	FName IconBrushName;
 };
 
+struct FHomeScreenLoadAtStartupSetting
+{
+public:
+	FHomeScreenLoadAtStartupSetting() = default;
+	FHomeScreenLoadAtStartupSetting(FOnSetLoadAtStartup InOnSetLoadAtStartup, FOnGetLoadAtStartup InOnGetLoadAtStartup)
+		: OnSetLoadAtStartupDelegate(InOnSetLoadAtStartup)
+		, OnGetLoadAtStartupDelegate(InOnGetLoadAtStartup)
+	{
+	}
+
+public:
+	FOnLoadAtStartupChanged& OnLoadAtStartupChanged() { return OnLoadAtStartupChangedDelegate; }
+	void SetOnLoadAtStartup(bool InAutoLoadOption) const { OnSetLoadAtStartupDelegate.ExecuteIfBound(InAutoLoadOption); }
+	bool GetOnLoadAtStartup() const { return OnGetLoadAtStartupDelegate.IsBound() ? OnGetLoadAtStartupDelegate.Execute() : false; }
+
+private:
+	FOnLoadAtStartupChanged OnLoadAtStartupChangedDelegate;
+	FOnSetLoadAtStartup OnSetLoadAtStartupDelegate;
+	FOnGetLoadAtStartup OnGetLoadAtStartupDelegate;
+};
+
+DECLARE_DELEGATE_RetVal(FHomeScreenLoadAtStartupSetting&, FOnGetStartupSetting)
+
 /** Setting used to create the HomeScreen, if leaved empty the defaults are used instead */
 struct FHomeScreenWidgetSettings
 {
@@ -57,6 +84,7 @@ struct FHomeScreenWidgetSettings
 		: bAlwaysShowRecentList(false)
 		, RecentProjectListWidget(SNullWidget::NullWidget)
 		, UnrealLogoBrush(nullptr)
+		, OnGetStartupSetting()
 	{
 	}
 
@@ -85,4 +113,9 @@ struct FHomeScreenWidgetSettings
 	 * Brush representing the icon for the Unreal Logo.
 	 */
 	const FSlateBrush* UnrealLogoBrush;
+
+	/**
+	 * Used to grab the StartupSetting logic to use
+	 */
+	FOnGetStartupSetting OnGetStartupSetting;
 };
diff --git a/Engine/Source/Editor/MainFrame/Public/Settings/HomeScreenSettings.h b/Engine/Source/Editor/MainFrame/Public/Settings/HomeScreenSettings.h
--- a/Engine/Source/Editor/MainFrame/Public/Settings/HomeScreenSettings.h
+++ b/Engine/Source/Editor/MainFrame/Public/Settings/HomeScreenSettings.h
@@ -5,8 +5,6 @@
 #include "UObject/Object.h"
 #include "HomeScreenSettings.generated.h"
 
-DECLARE_MULTICAST_DELEGATE_OneParam(FOnLoadAtStartupChanged, EAutoLoadProject)
-
 // Load on Startup setting is registered with the FEditorLoadingSavingSettingsCustomization in the Loading & Saving Section
 
 UCLASS(MinimalAPI, config=EditorSettings)
@@ -14,6 +12,9 @@ class UHomeScreenSettings : public UObject
 {
 	GENERATED_BODY()
 
+public:
+	UHomeScreenSettings();
+
 public:
 	UPROPERTY()
 	EAutoLoadProject LoadAtStartup = EAutoLoadProject::HomeScreen;
@@ -22,11 +23,17 @@ public:
 	MAINFRAME_API virtual void PostEditChangeProperty(FPropertyChangedEvent& PropertyChangedEvent) override;
 	MAINFRAME_API virtual void PostInitProperties() override;
 
+	UE_DEPRECATED(5.8, "Use the GetLoadAtStartupSetting instead")
 	FOnLoadAtStartupChanged& OnLoadAtStartupChanged() { return OnLoadAtStartupChangedDelegate; }
 
+	MAINFRAME_API FHomeScreenLoadAtStartupSetting& GetLoadAtStartupSetting();
+
 private:
 	void OnMostRecentProjectEditorSettingChanged(bool bInAutoLoad);
+	void SetLoadAtStartupChanged(bool bInAutoLoad);
+	bool GetLoadAtStartupChanged() const;
 
 private:
+	FHomeScreenLoadAtStartupSetting HomeScreenLoadAtStartupSetting;
 	FOnLoadAtStartupChanged OnLoadAtStartupChangedDelegate;
 };
