[HomeScreen] Timeout log issue in the home screen when checking for internet connection
//UE5/Main - Scheduled VP QAVirtualProduction Package and Test - Automation test errors in QAVirtualProduction Editor Win64

#Synced-CL 48654334

diff --git a/Engine/Source/Editor/MainFrame/Private/HomeScreen/SHomeScreen.cpp b/Engine/Source/Editor/MainFrame/Private/HomeScreen/SHomeScreen.cpp
--- a/Engine/Source/Editor/MainFrame/Private/HomeScreen/SHomeScreen.cpp
+++ b/Engine/Source/Editor/MainFrame/Private/HomeScreen/SHomeScreen.cpp
@@ -25,6 +25,7 @@
 #include "ToolMenus.h"
 #include "UnrealEdGlobals.h"
 #include "Internationalization/Culture.h"
+#include "Widgets/Docking/SDockTab.h"
 #include "Widgets/Images/SImage.h"
 #include "Widgets/Images/SThrobber.h"
 #include "Widgets/Input/SButton.h"
@@ -40,45 +41,6 @@
 
 #define LOCTEXT_NAMESPACE "SHomeScreen"
 
-SHomeScreen::SHomeScreen()
-{
-	if (!HttpRetryManager.IsValid())
-	{
-		HttpRetryManager = MakeShared<FHttpRetrySystem::FManager>(
-			FHttpRetrySystem::FRetryLimitCountSetting(3),
-			FHttpRetrySystem::FRetryTimeoutRelativeSecondsSetting(3));
-	}
-
-	if (GEditor && GEditor->IsTimerManagerValid())
-	{
-		constexpr float TimerManagerCallRate = 5.f;
-		constexpr bool bLoop = true;
-
-		GEditor->GetTimerManager()->SetTimer(
-			CheckInternetConnectionTimerHandle,
-			FTimerDelegate::CreateLambda([this] ()
-				{
-					if (bIsRequestFinished)
-					{
-						if (bIsConnected || bForceRetry)
-						{
-							bForceRetry = false;
-							TimerManagerCountOnceDisconnected = 0;
-							CheckInternetConnection();
-						}
-						else
-						{
-							// When disconnected retry connection less times than when you are connected
-							TimerManagerCountOnceDisconnected += 1;
-							bForceRetry = TimerManagerCountOnceDisconnected == MaxTimerManagerCountOnceDisconnectedBeforeRetry;
-						}
-					}
-				}),
-				TimerManagerCallRate,
-				bLoop);
-	}
-}
-
 SHomeScreen::~SHomeScreen()
 {
 	if (UHomeScreenSettings* HomeScreenSettings = GetMutableDefault<UHomeScreenSettings>())
@@ -86,10 +48,7 @@ SHomeScreen::~SHomeScreen()
 		HomeScreenSettings->OnLoadAtStartupChanged().RemoveAll(this);
 	}
 
-	if (GEditor && GEditor->IsTimerManagerValid())
-	{
-		GEditor->GetTimerManager()->ClearTimer(CheckInternetConnectionTimerHandle);
-	}
+	StopInternetConnectionTickCheck();
 
 	if (WebObject.IsValid())
 	{
@@ -103,6 +62,8 @@ SHomeScreen::~SHomeScreen()
 		HttpRetryRequest.Reset();
 	}
 
+	FGlobalTabmanager::Get()->OnTabForegrounded_Unsubscribe(OnTabForegroundedHandle);
+
 	HttpRetryManager.Reset();
 }
 
@@ -112,6 +73,15 @@ void SHomeScreen::Construct(const FArguments& InArgs, const TSharedPtr<SWebBrows
 	WebObject = InWebObject;
 	SectionLinks = InArgs._SectionLinks;
 
+	HttpRetryManager = MakeShared<FHttpRetrySystem::FManager>(
+		FHttpRetrySystem::FRetryLimitCountSetting(3),
+		FHttpRetrySystem::FRetryTimeoutRelativeSecondsSetting(3));
+
+	OnTabForegroundedHandle = FGlobalTabmanager::Get()->OnTabForegrounded_Subscribe(FOnActiveTabChanged::FDelegate::CreateSP(this, &SHomeScreen::OnTabForegrounded));
+
+	// Automatically start the internet connection check as if the HomeScreen is not a tab it won't ever start
+	StartInternetConnectionTickCheck();
+
 	AutoLoadProjectComboBoxSelection = EAutoLoadProject::HomeScreen;
 	if (UHomeScreenSettings* HomeScreenSettings = GetMutableDefault<UHomeScreenSettings>())
 	{
@@ -150,9 +120,6 @@ void SHomeScreen::Construct(const FArguments& InArgs, const TSharedPtr<SWebBrows
 	CreateSocialMediaButtons(TwitchButton, TEXT("https://www.twitch.tv/unrealengine"), FAppStyle::GetBrush("SocialMedia.Twitch"));
 	CreateSocialMediaButtons(InstagramButton, TEXT("https://www.instagram.com/unrealengine/"), FAppStyle::GetBrush("SocialMedia.Instagram"));
 
-	// Check for internet connection on startup
-	CheckInternetConnection();
-
 	ChildSlot
 	[
 		SNew(SBorder)
@@ -590,6 +557,27 @@ void SHomeScreen::Construct(const FArguments& InArgs, const TSharedPtr<SWebBrows
 	OnMainHomeSectionChanged(ECheckBoxState::Checked, MainHomeSelection);
 }
 
+void SHomeScreen::OnTabForegrounded(TSharedPtr<SDockTab> NewForegroundTab, TSharedPtr<SDockTab> BackgroundTab)
+{
+	// Stop the internet connection check if the foreground tab is not valid and return early
+	if (!NewForegroundTab.IsValid())
+	{
+		StopInternetConnectionTickCheck();
+		return;
+	}
+
+	// If the foreground tab contain this widget start its internet connection timer
+	// If it goes to the background then stop the internet connection
+	if (NewForegroundTab->GetContent() == AsShared())
+	{
+		StartInternetConnectionTickCheck();
+	}
+	else
+	{
+		StopInternetConnectionTickCheck();
+	}
+}
+
 FReply SHomeScreen::OnCreateProjectDialog(bool bInAllowProjectOpening, bool bInAllowProjectCreation)
 {
 	IMainFrameModule::Get().ShowProjectBrowser(bInAllowProjectOpening, bInAllowProjectCreation);
@@ -808,9 +796,7 @@ void SHomeScreen::CheckInternetConnection()
 		FHttpRetrySystem::FRetryTimeoutRelativeSecondsSetting(3),
 		FHttpRetrySystem::FRetryResponseCodes(),
 		FHttpRetrySystem::FRetryVerbs(),
-		FHttpRetrySystem::FRetryDomainsPtr(),
-		FHttpRetrySystem::FRetryLimitCountSetting(3),
-		FHttpRetrySystem::FExponentialBackoffCurve());
+		FHttpRetrySystem::FRetryDomainsPtr());
 
 	HttpRetryRequest->SetURL(TEXT("https://www.google.com/generate_204"));
 	HttpRetryRequest->SetVerb("GET");
@@ -1019,105 +1005,148 @@ bool SHomeScreen::HasAlreadyLatestEngineProject() const
 	return false;
 }
 
-FDelayedAutoRegisterHelper SHomeScreen::LoadStartupComboButtonRegistration(
-	EDelayedRegisterRunPhase::EndOfEngineInit,
-	[]
-		{
-			FToolMenuOwnerScoped OwnerScoped(UE_MODULE_NAME);
-			static const FName LoadStartupComboButtonName("HomeScreenComboButtonMenu");
-			UToolMenu* Menu = UToolMenus::Get()->RegisterMenu(LoadStartupComboButtonName);
-			FToolMenuSection& Section = Menu->AddSection("StartupSection");
-
-			Section.AddDynamicEntry(
-				TEXT("HomeTabDynamic"),
-				FNewToolMenuSectionDelegate::CreateLambda([] (FToolMenuSection& InSection)
+void SHomeScreen::StartInternetConnectionTickCheck()
+{
+	if (GEditor && GEditor->IsTimerManagerValid() && !CheckInternetConnectionTimerHandle.IsValid())
+	{
+		constexpr float TimerManagerCallRate = 5.f;
+		constexpr bool bLoop = true;
+
+		GEditor->GetTimerManager()->SetTimer(
+			CheckInternetConnectionTimerHandle,
+			FTimerDelegate::CreateSPLambda(this, [this]()
+				{
+					if (bIsRequestFinished)
 					{
-						if (UHomeScreenContext* Context = InSection.FindContext<UHomeScreenContext>())
+						if (bIsConnected || bForceRetry)
 						{
-							if (const TSharedPtr<SHomeScreen> HomeScreen = Context->HomeScreen.Pin())
-							{
-								const TSharedRef<SHomeScreen> HomeScreenRef = HomeScreen.ToSharedRef();
-
-								TSharedRef<SButton> HomeTabWidget = SNew(SButton)
-									.ButtonStyle(FAppStyle::Get(), "HomeScreen.ComboButton.MenuButton")
-									.ContentPadding(FMargin(12.f, 10.f))
-									.OnClicked(HomeScreenRef, &SHomeScreen::OnAutoLoadOptionChanged, EAutoLoadProject::HomeScreen)
-									[
-										SNew(SHorizontalBox)
-
-										+ SHorizontalBox::Slot()
-										.Padding(0.f, 0.f, 8.f, 0.f)
-										.AutoWidth()
-										[
-											SNew(SImage)
-											.Image(FAppStyle::GetBrush("Icons.Check"))
-											.Visibility(HomeScreenRef, &SHomeScreen::IsAutoLoadOptionCheckVisible, EAutoLoadProject::HomeScreen)
-										]
-
-										+ SHorizontalBox::Slot()
-										.AutoWidth()
-										[
-											SNew(STextBlock)
-											.Text(LOCTEXT("HomeScreenHomePanel", "Home Panel"))
-											.Font(FAppStyle::GetFontStyle("HomeScreen.Font.DefaultSize"))
-											.ColorAndOpacity(FStyleColors::White)
-										]
-									];
-
-								InSection.AddEntry(FToolMenuEntry::InitWidget(
-									TEXT("HomeTab"),
-									SNew(SBox).Padding(8.f, 0.f, 8.f, 4.f)[HomeTabWidget],
-									FText::GetEmpty()));
-							}
+							bForceRetry = false;
+							TimerManagerCountOnceDisconnected = 0;
+							CheckInternetConnection();
 						}
-					})
-				);
-
-
-			Section.AddDynamicEntry(
-				TEXT("MostRecentProjectDynamic"),
-				FNewToolMenuSectionDelegate::CreateLambda([] (FToolMenuSection& InSection)
-					{
-						if (UHomeScreenContext* Context = InSection.FindContext<UHomeScreenContext>())
+						else
 						{
-							if (const TSharedPtr<SHomeScreen> HomeScreen = Context->HomeScreen.Pin())
-							{
-								const TSharedRef<SHomeScreen> HomeScreenRef = HomeScreen.ToSharedRef();
+							// When disconnected retry connection less times than when you are connected
+							TimerManagerCountOnceDisconnected += 1;
+							bForceRetry = TimerManagerCountOnceDisconnected == MaxTimerManagerCountOnceDisconnectedBeforeRetry;
+						}
+					}
+				}),
+				TimerManagerCallRate,
+				bLoop);
+	}
+}
 
-								TSharedRef<SButton> MostRecentProjectWidget = SNew(SButton)
-									.ButtonStyle(FAppStyle::Get(), "HomeScreen.ComboButton.MenuButton")
-									.ContentPadding(FMargin(12.f, 10.f))
-									.OnClicked(HomeScreenRef, &SHomeScreen::OnAutoLoadOptionChanged, EAutoLoadProject::LastProject)
-									[
-										SNew(SHorizontalBox)
+void SHomeScreen::StopInternetConnectionTickCheck()
+{
+	if (GEditor && GEditor->IsTimerManagerValid() && CheckInternetConnectionTimerHandle.IsValid())
+	{
+		GEditor->GetTimerManager()->ClearTimer(CheckInternetConnectionTimerHandle);
+		CheckInternetConnectionTimerHandle.Invalidate();
+	}
+}
 
-										+ SHorizontalBox::Slot()
-										.Padding(0.f, 0.f, 8.f, 0.f)
-										.AutoWidth()
-										[
-											SNew(SImage)
-											.Image(FAppStyle::GetBrush("Icons.Check"))
-											.Visibility(HomeScreenRef, &SHomeScreen::IsAutoLoadOptionCheckVisible, EAutoLoadProject::LastProject)
-										]
+FDelayedAutoRegisterHelper SHomeScreen::LoadStartupComboButtonRegistration(
+	EDelayedRegisterRunPhase::EndOfEngineInit, 
+	[] {
+			UToolMenus::RegisterStartupCallback(FSimpleMulticastDelegate::FDelegate::CreateLambda([]()
+				{
+					FToolMenuOwnerScoped OwnerScoped(UE_MODULE_NAME);
+					static const FName LoadStartupComboButtonName("HomeScreenComboButtonMenu");
+					UToolMenu* Menu = UToolMenus::Get()->RegisterMenu(LoadStartupComboButtonName);
+					FToolMenuSection& Section = Menu->AddSection("StartupSection");
+
+					Section.AddDynamicEntry(
+						TEXT("HomeTabDynamic"),
+						FNewToolMenuSectionDelegate::CreateLambda([] (FToolMenuSection& InSection)
+							{
+								if (UHomeScreenContext* Context = InSection.FindContext<UHomeScreenContext>())
+								{
+									if (const TSharedPtr<SHomeScreen> HomeScreen = Context->HomeScreen.Pin())
+									{
+										const TSharedRef<SHomeScreen> HomeScreenRef = HomeScreen.ToSharedRef();
+
+										TSharedRef<SButton> HomeTabWidget = SNew(SButton)
+											.ButtonStyle(FAppStyle::Get(), "HomeScreen.ComboButton.MenuButton")
+											.ContentPadding(FMargin(12.f, 10.f))
+											.OnClicked(HomeScreenRef, &SHomeScreen::OnAutoLoadOptionChanged, EAutoLoadProject::HomeScreen)
+											[
+												SNew(SHorizontalBox)
+
+												+ SHorizontalBox::Slot()
+												.Padding(0.f, 0.f, 8.f, 0.f)
+												.AutoWidth()
+												[
+													SNew(SImage)
+													.Image(FAppStyle::GetBrush("Icons.Check"))
+													.Visibility(HomeScreenRef, &SHomeScreen::IsAutoLoadOptionCheckVisible, EAutoLoadProject::HomeScreen)
+												]
+
+												+ SHorizontalBox::Slot()
+												.AutoWidth()
+												[
+													SNew(STextBlock)
+													.Text(LOCTEXT("HomeScreenHomePanel", "Home Panel"))
+													.Font(FAppStyle::GetFontStyle("HomeScreen.Font.DefaultSize"))
+													.ColorAndOpacity(FStyleColors::White)
+												]
+											];
+
+										InSection.AddEntry(FToolMenuEntry::InitWidget(
+											TEXT("HomeTab"),
+											SNew(SBox).Padding(8.f, 0.f, 8.f, 4.f)[HomeTabWidget],
+											FText::GetEmpty()));
+									}
+								}
+							})
+						);
 
-										+ SHorizontalBox::Slot()
-										.AutoWidth()
-										[
-											SNew(STextBlock)
-											.Text(LOCTEXT("HomeScreenMostRecentProject", "Most Recent Project"))
-											.Font(FAppStyle::GetFontStyle("HomeScreen.Font.DefaultSize"))
-											.ColorAndOpacity(FStyleColors::White)
-										]
-									];
 
-								InSection.AddEntry(FToolMenuEntry::InitWidget(
-									TEXT("MostRecentProject"),
-									SNew(SBox).Padding(8.f, 0.f)[MostRecentProjectWidget],
-									FText::GetEmpty()));
-							}
-						}
-					})
-				);
+					Section.AddDynamicEntry(
+						TEXT("MostRecentProjectDynamic"),
+						FNewToolMenuSectionDelegate::CreateLambda([] (FToolMenuSection& InSection)
+							{
+								if (UHomeScreenContext* Context = InSection.FindContext<UHomeScreenContext>())
+								{
+									if (const TSharedPtr<SHomeScreen> HomeScreen = Context->HomeScreen.Pin())
+									{
+										const TSharedRef<SHomeScreen> HomeScreenRef = HomeScreen.ToSharedRef();
+
+										TSharedRef<SButton> MostRecentProjectWidget = SNew(SButton)
+											.ButtonStyle(FAppStyle::Get(), "HomeScreen.ComboButton.MenuButton")
+											.ContentPadding(FMargin(12.f, 10.f))
+											.OnClicked(HomeScreenRef, &SHomeScreen::OnAutoLoadOptionChanged, EAutoLoadProject::LastProject)
+											[
+												SNew(SHorizontalBox)
+
+												+ SHorizontalBox::Slot()
+												.Padding(0.f, 0.f, 8.f, 0.f)
+												.AutoWidth()
+												[
+													SNew(SImage)
+													.Image(FAppStyle::GetBrush("Icons.Check"))
+													.Visibility(HomeScreenRef, &SHomeScreen::IsAutoLoadOptionCheckVisible, EAutoLoadProject::LastProject)
+												]
+
+												+ SHorizontalBox::Slot()
+												.AutoWidth()
+												[
+													SNew(STextBlock)
+													.Text(LOCTEXT("HomeScreenMostRecentProject", "Most Recent Project"))
+													.Font(FAppStyle::GetFontStyle("HomeScreen.Font.DefaultSize"))
+													.ColorAndOpacity(FStyleColors::White)
+												]
+											];
+
+										InSection.AddEntry(FToolMenuEntry::InitWidget(
+											TEXT("MostRecentProject"),
+											SNew(SBox).Padding(8.f, 0.f)[MostRecentProjectWidget],
+											FText::GetEmpty()));
+									}
+								}
+							})
+						);
+				}));
 		}
 );
 
diff --git a/Engine/Source/Editor/MainFrame/Private/HomeScreen/SHomeScreen.h b/Engine/Source/Editor/MainFrame/Private/HomeScreen/SHomeScreen.h
--- a/Engine/Source/Editor/MainFrame/Private/HomeScreen/SHomeScreen.h
+++ b/Engine/Source/Editor/MainFrame/Private/HomeScreen/SHomeScreen.h
@@ -15,6 +15,7 @@
 class SButton;
 class SCheckBox;
 class SComboButton;
+class SDockTab;
 class SVerticalBox;
 class SWebBrowser;
 template <typename ItemType> class STileView;
@@ -36,12 +37,14 @@ public:
 		SLATE_ARGUMENT(const FSlateBrush*, UnrealLogoBrush)
 	SLATE_END_ARGS()
 
-	SHomeScreen();
 	virtual ~SHomeScreen() override;
 
 	/** Constructs the HomeScreen widget */
 	void Construct(const FArguments& InArgs, const TSharedPtr<SWebBrowser> InWebBrowser, const TWeakObjectPtr<UHomeScreenWeb> InWebObject);
 
+	/** Used when the HomeScreen tab becomes visible to start the internet connection check or to stop it when it goes in the background */
+	void OnTabForegrounded(TSharedPtr<SDockTab> NewForegroundTab, TSharedPtr<SDockTab> BackgroundTab);
+
 private:
 	/** Opens the Create Project dialog */
 	FReply OnCreateProjectDialog(bool bInAllowProjectOpening, bool bInAllowProjectCreation);
@@ -112,6 +115,12 @@ private:
 	/** Whether the current user already created project with this Engine version */
 	bool HasAlreadyLatestEngineProject() const;
 
+	/** Create and bind a function to check for internet connection from the TimerManager */
+	void StartInternetConnectionTickCheck();
+
+	/** Unbind the function to check for internet connection from the TimerManager */
+	void StopInternetConnectionTickCheck();
+
 private:
 	/** Registers the LoadStartup ComboButton Menu */
 	static FDelayedAutoRegisterHelper LoadStartupComboButtonRegistration;
@@ -126,6 +135,9 @@ private:
 	TSharedPtr<SCheckBox> GettingStartedCheckBox;
 	TSharedPtr<SCheckBox> SampleProjectsCheckBox;
 
+	/** Handle holding the OnTabForegrounded callback */
+	FDelegateHandle OnTabForegroundedHandle;
+
 	/** Http retry system used to check Internet connection */
 	TSharedPtr<FHttpRetrySystem::FManager> HttpRetryManager;
 	TSharedPtr<FHttpRetrySystem::FRequest> HttpRetryRequest;
