Enable xcode view debugger support for displaying Slate rendered Metal views in the editor

#Synced-CL 49225950

diff --git a/Engine/Source/Editor/MainFrame/Private/MainFrameModule.cpp b/Engine/Source/Editor/MainFrame/Private/MainFrameModule.cpp
--- a/Engine/Source/Editor/MainFrame/Private/MainFrameModule.cpp
+++ b/Engine/Source/Editor/MainFrame/Private/MainFrameModule.cpp
@@ -394,6 +394,10 @@ void FMainFrameModule::CreateDefaultMainFrameAuxiliary(const bool bStartImmersiv
 				.CreateTitleBar(!WindowConfig.bEmbedTitleAreaContent)
 				.SizingRule(WindowConfig.bIsUserSizable ? ESizingRule::UserSized : ESizingRule::FixedSize)
 				.SupportsMaximize(WindowConfig.bSupportsMaximize)
+#if PLATFORM_MAC
+				// this is needed to allow Slate-rendered Metal views to show up on XCode's view debugger
+				.SupportsTransparency(EWindowTransparency::PerPixel)
+#endif
 				.SupportsMinimize(WindowConfig.bSupportsMinimize);
 
 			const bool bShowRootWindowImmediately = false;
