[HomeScreen] Add display name for the Home Panel tab when navigating with Ctrl + Tab
Added display name to the HomeScreen panel, the name will be hidden while being a tab, but it will be shown if used in other tools, like the TabSwitcher feature (CTRL + TAB)

#Synced-CL 49539534

diff --git a/Engine/Source/Editor/MainFrame/Private/MainFrameModule.cpp b/Engine/Source/Editor/MainFrame/Private/MainFrameModule.cpp
--- a/Engine/Source/Editor/MainFrame/Private/MainFrameModule.cpp
+++ b/Engine/Source/Editor/MainFrame/Private/MainFrameModule.cpp
@@ -1142,6 +1142,8 @@ TSharedRef<SDockTab> FMainFrameModule::SpawnHomeScreenWindow(const FSpawnTabArgs
 {
 	TSharedRef<SDockTab> HomeScreenTab = SNew(SDockTab)
 		.TabRole(ETabRole::MajorTab)
+		// Assign a display name that will be used when switching tabs with the ctrl + tab feature
+		.Label(LOCTEXT("HomePanelDisplayName", "Home Panel"))
 		.IconColor(FStyleColors::White.GetSpecifiedColor())
 		.ContentPadding(FMargin(0));
 
