[HomeScreen] Without a home section the startup URL won't be changed
If the home section doesn't exist then use the first available section
If no section is valid the recent projects will be shown and no web page will be loaded

#Synced-CL 49539534

diff --git a/Engine/Source/Editor/MainFrame/Private/HomeScreen/SHomeScreen.cpp b/Engine/Source/Editor/MainFrame/Private/HomeScreen/SHomeScreen.cpp
--- a/Engine/Source/Editor/MainFrame/Private/HomeScreen/SHomeScreen.cpp
+++ b/Engine/Source/Editor/MainFrame/Private/HomeScreen/SHomeScreen.cpp
@@ -95,6 +95,12 @@ void SHomeScreen::Construct(const FArguments& InArgs, const TSharedPtr<SWebBrows
 	bHasLatestEngineProjects = InArgs._AlwaysShowRecentList || HasAlreadyLatestEngineProject();
 	MainHomeSelection = bHasLatestEngineProjects ? EMainSectionMenu::Home : EMainSectionMenu::GettingStarted;
 
+	// If the default home section is not found set it to none, so that at least the recent projects are shown
+	if (!SectionLinks.Contains(MainHomeSelection))
+	{
+		MainHomeSelection = EMainSectionMenu::None;
+	}
+
 	// Fixed Button Size
 	constexpr float SectionButtonCheckBoxHeight = 28.f;
 
@@ -445,7 +451,8 @@ void SHomeScreen::Construct(const FArguments& InArgs, const TSharedPtr<SWebBrows
 				SNew(SWidgetSwitcher)
 				.WidgetIndex_Lambda([this] ()
 				{
-					return MainHomeSelection == EMainSectionMenu::Home && bHasLatestEngineProjects ? /** Show Recent Projects */ 0 : /** Show just the WebPage */ 1;
+					// If the main home selection is none or if you are in the Home section and have the latest engine projects show the RecentProjects
+					return MainHomeSelection == EMainSectionMenu::None || (MainHomeSelection == EMainSectionMenu::Home && bHasLatestEngineProjects) ? /** Show Recent Projects */ 0 : /** Show just the WebPage */ 1;
 				})
 
 				+ SWidgetSwitcher::Slot()
