[HomeScreen] Can open multiple New Project and Project Browser modal windows
If there is already an existing project browser window open while using the FMainFrameActionCallbacks NewProject it will focus the already created window.
The function is used for the ShowProjectBrowser function in the MainFrameModule used currently by the HomeScreen and for the menu entry in the file and project menu.

#Synced-CL 49539534

diff --git a/Engine/Source/Editor/MainFrame/Private/Frame/MainFrameActions.cpp b/Engine/Source/Editor/MainFrame/Private/Frame/MainFrameActions.cpp
--- a/Engine/Source/Editor/MainFrame/Private/Frame/MainFrameActions.cpp
+++ b/Engine/Source/Editor/MainFrame/Private/Frame/MainFrameActions.cpp
@@ -66,6 +66,7 @@ DEFINE_LOG_CATEGORY_STATIC(MainFrameActions, Log, All);
 TSharedRef< FUICommandList > FMainFrameCommands::ActionList( new FUICommandList() );
 
 TWeakPtr<SNotificationItem> FMainFrameActionCallbacks::ChoosePackagesToCheckInNotification;
+TWeakPtr<SWindow> FMainFrameActionCallbacks::NewProjectWindowWeak;
 
 namespace
 {
@@ -424,10 +425,20 @@ void FMainFrameActionCallbacks::NewProject( bool bAllowProjectOpening, bool bAll
 		return;
 	}
 
+	if (const TSharedPtr<SWindow> ExistingWindow = NewProjectWindowWeak.Pin())
+	{
+		ExistingWindow->ShowWindow();
+		ExistingWindow->BringToFront(/* bForce */ false);
+		ExistingWindow->FlashWindow();
+		return;
+	}
+
 	TSharedRef<SWindow> NewProjectWindow = CreateProjectBrowserWindow(bAllowProjectOpening, bAllowProjectCreate);
 
 	NewProjectWindow->SetContent( FGameProjectGenerationModule::Get().CreateGameProjectDialog(bAllowProjectOpening, bAllowProjectCreate) );
 
+	NewProjectWindowWeak = NewProjectWindow;
+
 	IMainFrameModule& MainFrameModule = FModuleManager::GetModuleChecked<IMainFrameModule>(TEXT("MainFrame"));
 	if (MainFrameModule.GetParentWindow().IsValid())
 	{
diff --git a/Engine/Source/Editor/MainFrame/Private/Frame/MainFrameActions.h b/Engine/Source/Editor/MainFrame/Private/Frame/MainFrameActions.h
--- a/Engine/Source/Editor/MainFrame/Private/Frame/MainFrameActions.h
+++ b/Engine/Source/Editor/MainFrame/Private/Frame/MainFrameActions.h
@@ -324,4 +324,7 @@ private:
 
 	/** The notification in place while we choose packages to check in */
 	static TWeakPtr<SNotificationItem> ChoosePackagesToCheckInNotification;
+	
+	/** New project window weak ptr, used to give the window focus if trying to open a new window when there is already an existing one */
+	static TWeakPtr<SWindow> NewProjectWindowWeak;
 };
