[HomeScreen] Change URL localization and links to be always up to date

#Synced-CL 49923896

diff --git a/Engine/Source/Editor/MainFrame/Private/HomeScreen/SHomeScreen.cpp b/Engine/Source/Editor/MainFrame/Private/HomeScreen/SHomeScreen.cpp
--- a/Engine/Source/Editor/MainFrame/Private/HomeScreen/SHomeScreen.cpp
+++ b/Engine/Source/Editor/MainFrame/Private/HomeScreen/SHomeScreen.cpp
@@ -591,7 +591,17 @@ FReply SHomeScreen::OnCreateProjectDialog(bool bInAllowProjectOpening, bool bInA
 	return FReply::Handled();
 }
 
-FReply SHomeScreen::OnSocialMediaClicked(FString InURL)
+FReply SHomeScreen::OnLinkClicked(TAttribute<FString> InURL) const
+{
+	const FString URL = InURL.Get();
+	if (!URL.IsEmpty())
+	{
+		return OnLinkClicked(URL);
+	}
+	return FReply::Handled();
+}
+
+FReply SHomeScreen::OnLinkClicked(FString InURL) const
 {
 	FPlatformProcess::LaunchURL(*InURL, nullptr, nullptr);
 	return FReply::Handled();
@@ -720,13 +730,13 @@ void SHomeScreen::CreateMainSectionCheckBox(TSharedPtr<SCheckBox>& OutCheckBox,
 		]);
 }
 
-void SHomeScreen::CreateResourceButtons(TSharedPtr<SButton>& OutButton, FString InLink, const FText& InText, const FSlateBrush* InImage)
+void SHomeScreen::CreateResourceButtons(TSharedPtr<SButton>& OutButton, TAttribute<FString> InLink, const FText& InText, const FSlateBrush* InImage)
 {
 	SAssignNew(OutButton, SButton)
 	.IsEnabled(this, &SHomeScreen::IsConnectedToInternet)
 	.ButtonStyle(FAppStyle::Get(), "NoBorder")
 	.ContentPadding(FMargin(0.f, 2.f))
-	.OnClicked(this, &SHomeScreen::OnSocialMediaClicked, InLink);
+	.OnClicked(this, &SHomeScreen::OnLinkClicked, InLink);
 
 	OutButton->SetContent(
 		SNew(SHorizontalBox)
@@ -756,7 +766,7 @@ void SHomeScreen::CreateSocialMediaButtons(TSharedPtr<SButton>& OutButton, FStri
 {
 	SAssignNew(OutButton, SButton)
 	.ButtonStyle(FAppStyle::Get(), "NoBorder")
-	.OnClicked(this, &SHomeScreen::OnSocialMediaClicked, InLink);
+	.OnClicked(this, &SHomeScreen::OnLinkClicked, InLink);
 
 	OutButton->SetContent(
 			SNew(SImage)
diff --git a/Engine/Source/Editor/MainFrame/Private/HomeScreen/SHomeScreen.h b/Engine/Source/Editor/MainFrame/Private/HomeScreen/SHomeScreen.h
--- a/Engine/Source/Editor/MainFrame/Private/HomeScreen/SHomeScreen.h
+++ b/Engine/Source/Editor/MainFrame/Private/HomeScreen/SHomeScreen.h
@@ -51,8 +51,11 @@ private:
 	/** Opens the Create Project dialog */
 	FReply OnCreateProjectDialog(bool bInAllowProjectOpening, bool bInAllowProjectCreation);
 
-	/** Handles social media button clicks */
-	FReply OnSocialMediaClicked(FString InURL);
+	/** Handle the clicked links event */
+	FReply OnLinkClicked(TAttribute<FString> InURL) const;
+	
+	/** Handle the clicked links event */
+	FReply OnLinkClicked(FString InURL) const;
 
 	/** Checks if a main home section is selected */
 	ECheckBoxState IsMainHomeSectionChecked(EMainSectionMenu InMainHomeSelection) const;
@@ -82,7 +85,7 @@ private:
 	void CreateMainSectionCheckBox(TSharedPtr<SCheckBox>& OutCheckBox, EMainSectionMenu InMainHomeSelection, const FText& InText, const FSlateBrush* InImage);
 
 	/** Creates resource buttons */
-	void CreateResourceButtons(TSharedPtr<SButton>& OutButton, FString InLink, const FText& InText, const FSlateBrush* InImage);
+	void CreateResourceButtons(TSharedPtr<SButton>& OutButton, TAttribute<FString> InLink, const FText& InText, const FSlateBrush* InImage);
 
 	/** Creates social media buttons */
 	void CreateSocialMediaButtons(TSharedPtr<SButton>& OutButton, FString InLink, const FSlateBrush* InImage);
diff --git a/Engine/Source/Editor/MainFrame/Private/MainFrameModule.cpp b/Engine/Source/Editor/MainFrame/Private/MainFrameModule.cpp
--- a/Engine/Source/Editor/MainFrame/Private/MainFrameModule.cpp
+++ b/Engine/Source/Editor/MainFrame/Private/MainFrameModule.cpp
@@ -39,14 +39,15 @@
 // HomeScreen Includes
 #include "HomeScreen/HomeScreenWeb.h"
 #include "HomeScreen/SHomeScreen.h" 
+#include "Internationalization/Culture.h"
 #include "IWebBrowserSingleton.h"
 #include "IWebBrowserWindow.h"
 #include "Menus/SettingsMenu.h"
 #include "Settings/EditorStyleSettings.h"
+#include "Settings/HomeScreenSettings.h"
 #include "SWebBrowser.h"
 #include "ToolMenus.h"
 #include "WebBrowserModule.h"
-#include "Settings/HomeScreenSettings.h"
 #include "Widgets/Docking/SDockTab.h"
 
 DEFINE_LOG_CATEGORY(LogMainFrame);
@@ -100,38 +101,76 @@ namespace HomeScreenHelper
 	{
 		TArray<FHomeScreenResourceCommunityEntry> DefaultEntries;
 
+		FString EngineVersionString = FEngineVersion::Current().ToString(EVersionComponent::Minor);
+		EngineVersionString = EngineVersionString.Replace(TEXT("."), TEXT("-"));
+
 		// Forum
 		FHomeScreenResourceCommunityEntry ForumEntry;
 		ForumEntry.DisplayName = LOCTEXT("ProjectHomeForumsButton", "Forums");
-		ForumEntry.Link = TEXT("https://forums.unrealengine.com/categories?tag=unreal-engine");
+		ForumEntry.Link = TAttribute<FString>::CreateLambda([]()
+			{
+				const FString FullLocale = FInternationalization::Get().GetCurrentLanguage()->GetName();
+				const FString ForumLink = 
+					TEXT("https://forums.unrealengine.com/categories?tag=unreal-engine&locale=")
+					+ FullLocale;
+				return ForumLink;
+			});
 		ForumEntry.IconBrushName = TEXT("HomeScreen.Forum");
 		DefaultEntries.Add(ForumEntry);
 
 		// Documentation
 		FHomeScreenResourceCommunityEntry DocumentationEntry;
 		DocumentationEntry.DisplayName = LOCTEXT("ProjectHomeDocumentation", "Documentation");
-		DocumentationEntry.Link = TEXT("https://dev.epicgames.com/documentation/en-us/unreal-engine/unreal-engine-5-6-documentation");
+		DocumentationEntry.Link = TAttribute<FString>::CreateLambda([EngineVersionString]()
+			{
+				const FString FullLocale = FInternationalization::Get().GetCurrentLanguage()->GetName();
+				const FString DocumentationLink = 
+					TEXT("https://dev.epicgames.com/documentation/")
+					+ FullLocale
+					+ TEXT("/unreal-engine/unreal-engine-")
+					+  EngineVersionString
+					+ TEXT("-documentation");
+				return DocumentationLink;
+			});
 		DocumentationEntry.IconBrushName = TEXT("HomeScreen.Documentation");
 		DefaultEntries.Add(DocumentationEntry);
 
 		// Tutorial
 		FHomeScreenResourceCommunityEntry TutorialEntry;
 		TutorialEntry.DisplayName = LOCTEXT("ProjectHomeTutorialsButton", "Tutorials");
-		TutorialEntry.Link = TEXT("https://dev.epicgames.com/community/unreal-engine/learning");
+		TutorialEntry.Link = TAttribute<FString>::CreateLambda([]()
+			{
+				const FString FullLocale = FInternationalization::Get().GetCurrentLanguage()->GetName();
+				const FString TutorialLink = 
+					TEXT("https://dev.epicgames.com/community/unreal-engine/learning?locale=")
+					+ FullLocale;
+				return TutorialLink;
+			});
 		TutorialEntry.IconBrushName = TEXT("HomeScreen.Tutorial");
 		DefaultEntries.Add(TutorialEntry);
 
 		// Roadmap
 		FHomeScreenResourceCommunityEntry RoadmapEntry;
 		RoadmapEntry.DisplayName = LOCTEXT("ProjectHomeRoadmapButton", "Roadmap");
-		RoadmapEntry.Link = TEXT("https://portal.productboard.com/epicgames/1-unreal-engine-public-roadmap/tabs/109-unreal-engine-5-5");
+		// This link will always point to the latest roadmap and its intended to always point to latest
+		RoadmapEntry.Link = TEXT("https://unrealengine.com/roadmap");
 		RoadmapEntry.IconBrushName = TEXT("HomeScreen.Roadmap");
 		DefaultEntries.Add(RoadmapEntry);
 
 		// Release
-		FHomeScreenResourceCommunityEntry ReleaseEntry;
+		FHomeScreenResourceCommunityEntry ReleaseEntry; 
 		ReleaseEntry.DisplayName = LOCTEXT("ProjectHomeReleaseNotesButton", "Release Notes");
-		ReleaseEntry.Link = TEXT("https://dev.epicgames.com/documentation/en-us/unreal-engine/unreal-engine-5-6-release-notes");
+		ReleaseEntry.Link = TAttribute<FString>::CreateLambda([EngineVersionString]()
+			{
+				const FString FullLocale = FInternationalization::Get().GetCurrentLanguage()->GetName();
+				const FString ReleaseLink = 
+					TEXT("https://dev.epicgames.com/documentation/")
+					+ FullLocale
+					+ TEXT("/unreal-engine/unreal-engine-")
+					+  EngineVersionString
+					+ TEXT("-release-notes");
+				return ReleaseLink;
+			});
 		ReleaseEntry.IconBrushName = TEXT("HomeScreen.ReleaseNotes");
 		DefaultEntries.Add(ReleaseEntry);
 		
diff --git a/Engine/Source/Editor/MainFrame/Public/Settings/HomeScreenCommon.h b/Engine/Source/Editor/MainFrame/Public/Settings/HomeScreenCommon.h
--- a/Engine/Source/Editor/MainFrame/Public/Settings/HomeScreenCommon.h
+++ b/Engine/Source/Editor/MainFrame/Public/Settings/HomeScreenCommon.h
@@ -49,7 +49,7 @@ using FHomeScreenSectionToLinkMap = TMap<EMainSectionMenu, FHomeScreenSectionLin
  */
 struct FHomeScreenResourceCommunityEntry
 {
-	FString Link;
+	TAttribute<FString> Link;
 	FText DisplayName;
 	FName IconBrushName;
 };
