[HomeScreen] Revisit previous implementation for the MainFrameModule ProjectBrowser hooks

#Synced-CL 49992949

diff --git a/Engine/Source/Editor/MainFrame/Private/MainFrameModule.cpp b/Engine/Source/Editor/MainFrame/Private/MainFrameModule.cpp
--- a/Engine/Source/Editor/MainFrame/Private/MainFrameModule.cpp
+++ b/Engine/Source/Editor/MainFrame/Private/MainFrameModule.cpp
@@ -196,13 +196,14 @@ namespace HomeScreenHelper
 
 	TSharedRef<SWidget> GetRecentProjectListWidget()
 	{
-		constexpr bool bOnlyShowList = true;
-		constexpr bool bProjectBrowserAllowScrolling = false;
+		constexpr bool bIsEmbedded = true;
 		return SNew(SBox)
 			.HeightOverride(172.f)
 			.Padding(0.f)
 			[
-				IMainFrameModule::Get().CreateProjectBrowser(bOnlyShowList, bProjectBrowserAllowScrolling)
+				// Embedded should not show the HeaderBar and should not allow scrolling
+				FGameProjectGenerationModule::Get()
+					.CreateProjectBrowser(!bIsEmbedded /** ShowHeaderBar */, !bIsEmbedded /** AllowScrolling */)
 			];
 	}
 
@@ -876,16 +877,6 @@ void FMainFrameModule::ShowProjectBrowser(bool bAllowProjectOpening, bool bAllow
 	FMainFrameActionCallbacks::NewProject(bAllowProjectOpening, bAllowProjectCreate);
 }
 
-TSharedRef<SWidget> FMainFrameModule::CreateProjectBrowser(bool bOnlyShowList, bool bAllowScrolling)
-{
-	if (CreateProjectBrowserOverride.IsBound())
-	{
-		return CreateProjectBrowserOverride.Execute(bOnlyShowList, bAllowScrolling);
-	}
-
-	return FGameProjectGenerationModule::Get().CreateProjectBrowser(!bOnlyShowList, bAllowScrolling);
-}
-
 void FMainFrameModule::StartupModule( )
 {
 	bDelayedShowMainFrame = false;
diff --git a/Engine/Source/Editor/MainFrame/Private/MainFrameModule.h b/Engine/Source/Editor/MainFrame/Private/MainFrameModule.h
--- a/Engine/Source/Editor/MainFrame/Private/MainFrameModule.h
+++ b/Engine/Source/Editor/MainFrame/Private/MainFrameModule.h
@@ -228,9 +228,6 @@ public:
 	// Opens the project browser in a new window
 	virtual void ShowProjectBrowser(bool bAllowProjectOpening, bool bAllowProjectCreate) override;
 
-	// Creates the project browser for embedding
-	virtual TSharedRef<SWidget> CreateProjectBrowser(bool bOnlyShowList = true, bool bAllowScrolling = true) override;
-
 	// Overrides the functionality of ShowProjectBrowser with the provided delegate
 	virtual FDelegateHandle RegisterShowProjectBrowserOverride(const FShowProjectBrowser& InShowProjectBrowserOverride) override
 	{
@@ -247,22 +244,6 @@ public:
 		}
 	}
 
-	// Overrides the functionality of CreateProjectBrowser with the provided delegate
-	virtual FDelegateHandle RegisterCreateProjectBrowserOverride(const FCreateProjectBrowser& InCreateProjectBrowserOverride) override
-	{
-		CreateProjectBrowserOverride = InCreateProjectBrowserOverride;
-		return CreateProjectBrowserOverride.GetHandle();
-	}
-
-	// Removes the override of CreateProjectBrowser if it matches the handle
-	virtual void UnregisterCreateProjectBrowserOverride(FDelegateHandle InHandle) override
-	{
-		if (CreateProjectBrowserOverride.GetHandle() == InHandle)
-		{
-			CreateProjectBrowserOverride = FCreateProjectBrowser{};
-		}
-	}
-
 	// IModuleInterface interface
 
 	virtual void StartupModule( ) override;
@@ -402,7 +383,4 @@ private:
 
 	// Stores the delegate for the override of ShowProjectBrowser, if any
 	FShowProjectBrowser ShowProjectBrowserOverride;
-
-	// Stores the delegate for the override of CreateProjectBrowser, if any
-	FCreateProjectBrowser CreateProjectBrowserOverride;
 };
diff --git a/Engine/Source/Editor/MainFrame/Public/Interfaces/IMainFrameModule.h b/Engine/Source/Editor/MainFrame/Public/Interfaces/IMainFrameModule.h
--- a/Engine/Source/Editor/MainFrame/Public/Interfaces/IMainFrameModule.h
+++ b/Engine/Source/Editor/MainFrame/Public/Interfaces/IMainFrameModule.h
@@ -29,7 +29,6 @@ struct FMainFrameDeveloperTool
 };
 
 FUNC_DECLARE_DELEGATE(FShowProjectBrowser, void, bool /*bAllowProjectOpening*/, bool /*bAllowProjectCreate*/);
-FUNC_DECLARE_DELEGATE(FCreateProjectBrowser, TSharedRef<SWidget>, bool /*bOnlyShowList*/, bool /*bAllowScrolling*/);
 
 /**
  * Interface for main frame modules.
@@ -255,11 +254,6 @@ public:
 	 */
 	virtual void ShowProjectBrowser(bool bAllowProjectOpening, bool bAllowProjectCreate) = 0;
 
-	/**
-	 * Creates the project browser for embedding.
-	 */
-	virtual TSharedRef<SWidget> CreateProjectBrowser(bool bOnlyShowList = true, bool bAllowScrolling = true) = 0;
-
 	/**
 	 * Overrides the functionality of ShowProjectBrowser with the provided delegate.
 	 */
@@ -270,16 +264,6 @@ public:
 	 */
 	virtual void UnregisterShowProjectBrowserOverride(FDelegateHandle InHandle) = 0;
 
-	/**
-	 * Overrides the functionality of CreateProjectBrowser with the provided delegate.
-	 */
-	virtual FDelegateHandle RegisterCreateProjectBrowserOverride(const FCreateProjectBrowser& InCreateProjectBrowserOverride) = 0;
-
-	/**
-	 * Removes the override of CreateProjectBrowser if it matches the handle.
-	 */
-	virtual void UnregisterCreateProjectBrowserOverride(FDelegateHandle InHandle) = 0;
-
 	/**
 	 * Return the HomeScreen tab label
 	 */
