[HomeScreen] Polish various elements with padding/styling

#Synced-CL 49992949

diff --git a/Engine/Source/Editor/EditorStyle/Private/StarshipStyle.cpp b/Engine/Source/Editor/EditorStyle/Private/StarshipStyle.cpp
--- a/Engine/Source/Editor/EditorStyle/Private/StarshipStyle.cpp
+++ b/Engine/Source/Editor/EditorStyle/Private/StarshipStyle.cpp
@@ -8559,7 +8559,7 @@ void FStarshipEditorStyle::FStyle::SetupHomeScreenStyles()
 	FSlateFontInfo DefaultBoldFontSize = DEFAULT_FONT("Bold", 12);
 	Set("HomeScreen.Font.DefaultBoldSize", DefaultBoldFontSize);
 
-	FSlateFontInfo RecentProjectFont = DEFAULT_FONT("Bold", 20);
+	FSlateFontInfo RecentProjectFont = DEFAULT_FONT("Bold", 18);
 	Set("HomeScreen.Font.RecentProject", RecentProjectFont);
 
 	FComboButtonStyle LoadOnStartupStyle = FAppStyle::Get().GetWidgetStyle<FComboButtonStyle>("ComboButton");
@@ -8604,11 +8604,11 @@ void FStarshipEditorStyle::FStyle::SetupHomeScreenStyles()
 	Set("HomeScreen.MyFolderButton", MyFolderButton);
 
 	FButtonStyle SeeAllProjectsButton = FButtonStyle()
-		.SetNormal(FSlateRoundedBoxBrush(FStyleColors::Transparent, 4.0f, COLOR("FFFFFF59"), 1.f))
-		.SetHovered(FSlateRoundedBoxBrush(COLOR("FFFFFF0D"), 4.0f, COLOR("FFFFFF59"), 1.f))
-		.SetPressed(FSlateRoundedBoxBrush(COLOR("FFFFFF1A"), 4.0f, COLOR("FFFFFF59"), 1.f))
+		.SetNormal(FSlateRoundedBoxBrush(FStyleColors::Transparent, 8.0f, COLOR("565659FF"), 1.f))
+		.SetHovered(FSlateRoundedBoxBrush(COLOR("404044FF"), 8.0f, COLOR("FFFFFFFF"), 1.f))
+		.SetPressed(FSlateRoundedBoxBrush(COLOR("404044FF"), 8.0f, COLOR("FFFFFFFF"), 1.f))
 		.SetHoveredForeground(FStyleColors::White)
-		.SetNormalForeground(FStyleColors::White)
+		.SetNormalForeground(COLOR("EDEDEDFF"))
 		.SetPressedForeground(FStyleColors::White)
 		.SetNormalPadding(FMargin(0.f, 4.0f))
 		.SetPressedPadding(FMargin(0.f, 4.0f));
diff --git a/Engine/Source/Editor/MainFrame/Private/HomeScreen/SHomeScreen.cpp b/Engine/Source/Editor/MainFrame/Private/HomeScreen/SHomeScreen.cpp
--- a/Engine/Source/Editor/MainFrame/Private/HomeScreen/SHomeScreen.cpp
+++ b/Engine/Source/Editor/MainFrame/Private/HomeScreen/SHomeScreen.cpp
@@ -476,7 +476,7 @@ void SHomeScreen::Construct(const FArguments& InArgs, const TSharedPtr<SWebBrows
 						[
 							SNew(SVerticalBox)
 							+ SVerticalBox::Slot()
-							.Padding(0.f, 64.f, 0.f, 16.f)
+							.Padding(0.f, 64.f, 0.f, 48.f)
 							[
 								SNew(SHorizontalBox)
 							
@@ -490,19 +490,25 @@ void SHomeScreen::Construct(const FArguments& InArgs, const TSharedPtr<SWebBrows
 									.Font(FAppStyle::GetFontStyle("HomeScreen.Font.RecentProject"))
 									.ColorAndOpacity(FStyleColors::White)
 								]
-							
+
+								+ SHorizontalBox::Slot()
+								[
+									SNew(SSpacer)
+								]
+
 								+ SHorizontalBox::Slot()
 								.AutoWidth()
 								.VAlign(VAlign_Center)
+								.HAlign(HAlign_Right)
 								[
 									SNew(SButton)
 									.ButtonStyle(FAppStyle::Get(), "HomeScreen.SeeAllProjectsButton")
 									.OnClicked(this, &SHomeScreen::OnCreateProjectDialog, /** bInAllowProjectOpening */true, /** bInAllowProjectCreation */false)
-									.ContentPadding(FMargin(18, 10))
+									.ContentPadding(FMargin(16, 8))
 									[
 										SNew(STextBlock)
-										.Font(FAppStyle::GetFontStyle("HomeScreen.Font.DefaultSize"))
-										.Text(LOCTEXT("SeeAllProjects", "See All Projects"))
+										.Font(FAppStyle::GetFontStyle("HomeScreen.Font.DefaultBoldSize"))
+										.Text(LOCTEXT("SeeAllProjects", "See all projects"))
 									]
 								]
 							]
@@ -515,12 +521,7 @@ void SHomeScreen::Construct(const FArguments& InArgs, const TSharedPtr<SWebBrows
 								.BorderImage(FAppStyle::GetBrush("HomeScreen.RecentProjects.Background"))
 								.Padding(24.f)
 								[
-									SNew(SVerticalBox)
-
-									+ SVerticalBox::Slot()
-									[
-										InArgs._RecentProjectList
-									]
+									InArgs._RecentProjectList
 								]
 							]
 						]
diff --git a/Engine/Source/Editor/MainFrame/Private/MainFrameModule.cpp b/Engine/Source/Editor/MainFrame/Private/MainFrameModule.cpp
--- a/Engine/Source/Editor/MainFrame/Private/MainFrameModule.cpp
+++ b/Engine/Source/Editor/MainFrame/Private/MainFrameModule.cpp
@@ -198,12 +198,15 @@ namespace HomeScreenHelper
 	{
 		constexpr bool bIsEmbedded = true;
 		return SNew(SBox)
-			.HeightOverride(172.f)
-			.Padding(0.f)
+			.Padding(8.f)
 			[
-				// Embedded should not show the HeaderBar and should not allow scrolling
-				FGameProjectGenerationModule::Get()
-					.CreateProjectBrowser(!bIsEmbedded /** ShowHeaderBar */, !bIsEmbedded /** AllowScrolling */)
+				SNew(SBox)
+				.HeightOverride(172.f)
+				[
+					// Embedded should not show the HeaderBar and should not allow scrolling
+					FGameProjectGenerationModule::Get()
+						.CreateProjectBrowser(!bIsEmbedded /** ShowHeaderBar */, !bIsEmbedded /** AllowScrolling */)
+				]
 			];
 	}
 
