Logic to revert to the default layout is broken. Hard reset to the default layout on fail.

#Synced-CL 50306593

diff --git a/Engine/Source/Editor/MainFrame/Private/MainFrameModule.cpp b/Engine/Source/Editor/MainFrame/Private/MainFrameModule.cpp
--- a/Engine/Source/Editor/MainFrame/Private/MainFrameModule.cpp
+++ b/Engine/Source/Editor/MainFrame/Private/MainFrameModule.cpp
@@ -573,8 +573,11 @@ void FMainFrameModule::CreateDefaultMainFrameAuxiliary(const bool bStartImmersiv
 			const bool bPrimaryAreaMustHaveOpenedTabsToBeValid = true;
 			const EOutputCanBeNullptr OutputCanBeNullptr = EOutputCanBeNullptr::IfNoOpenTabValid;
 			TArray<FString> RemovedOlderLayoutVersions;
-			const TSharedRef<FTabManager::FLayout> LoadedLayout = FLayoutSaveRestore::LoadFromConfig(GEditorLayoutIni,
-				DefaultLayout, OutputCanBeNullptr, RemovedOlderLayoutVersions);
+
+			const TSharedRef<FTabManager::FLayout> LoadedLayout = bIsReCreatingDefaultLayoutAsFallback
+				? DefaultLayout
+				: FLayoutSaveRestore::LoadFromConfig(GEditorLayoutIni, DefaultLayout, OutputCanBeNullptr, RemovedOlderLayoutVersions);
+
 
 			// If older fields of the layout name (i.e., lower versions than "UnrealEd_Layout_v1.4") were found
 			if (RemovedOlderLayoutVersions.Num() > 0)
@@ -596,6 +599,7 @@ void FMainFrameModule::CreateDefaultMainFrameAuxiliary(const bool bStartImmersiv
 		
 			MainFrameContent = FGlobalTabmanager::Get()->RestoreFrom(LoadedLayout, RootWindow, WindowConfig.bEmbedTitleAreaContent, OutputCanBeNullptr);
 			// MainFrameContent will only be nullptr if its main area contains invalid tabs (probably some layout bug). If so, reset layout to avoid potential crashes
+		
 			if (!MainFrameContent.IsValid())
 			{
 				// This code will recurse to load the default layout. If that also fail we will get into an infinite loop
