Bump the editor layout version, I added a new concept for the window placement.

#Synced-CL 50202233

diff --git a/Engine/Source/Editor/MainFrame/Private/MainFrameModule.cpp b/Engine/Source/Editor/MainFrame/Private/MainFrameModule.cpp
--- a/Engine/Source/Editor/MainFrame/Private/MainFrameModule.cpp
+++ b/Engine/Source/Editor/MainFrame/Private/MainFrameModule.cpp
@@ -431,7 +431,6 @@ void FMainFrameModule::CreateDefaultMainFrameAuxiliary(const bool bStartImmersiv
 				.AutoCenter(WindowConfig.CenterRules)
 				.Title(WindowConfig.WindowTitle)
 				.CloseButtonToolTipText(CloseButtonToolTipText)
-				.IsInitiallyMaximized(DefaultWindowLocation.InitiallyMaximized)
 				.ScreenPosition(DefaultWindowLocation.ScreenPosition)
 				.ClientSize(DefaultWindowLocation.WindowSize)
 				.CreateTitleBar(!WindowConfig.bEmbedTitleAreaContent)
@@ -443,6 +442,12 @@ void FMainFrameModule::CreateDefaultMainFrameAuxiliary(const bool bStartImmersiv
 #endif
 				.SupportsMinimize(WindowConfig.bSupportsMinimize);
 
+			// Avoid using the is initially maximized argument when creating the window so that the code below can change the window without this initial value taking over taking the changes when the window is first shown.
+			if (DefaultWindowLocation.InitiallyMaximized)
+			{
+				RootWindow->Maximize();
+			}
+
 			const bool bShowRootWindowImmediately = false;
 			FSlateApplication::Get().AddWindow( RootWindow, bShowRootWindowImmediately );
 
@@ -511,12 +516,12 @@ void FMainFrameModule::CreateDefaultMainFrameAuxiliary(const bool bStartImmersiv
 				DisplayMetrics.PrimaryDisplayWorkAreaRect.Bottom - DisplayMetrics.PrimaryDisplayWorkAreaRect.Top );
 			const FVector2D WindowSize = (CenterScale * DisplaySize) / DPIScale;
 
-			// IMPORTANT: If you want to change the default value of "LevelEditor_Layout_v1.1" or "UnrealEd_Layout_v1.4" (even if you only change their version numbers), these are the steps to follow:
+			// IMPORTANT: If you want to change the default value of "LevelEditor_Layout_v1.1" or "UnrealEd_Layout_v1.6" (even if you only change their version numbers), these are the steps to follow:
 			// 1. Check out Engine\Config\Layouts\DefaultLayout.ini in Perforce.
 			// 2. Change the code below as you wish and compile the code.
 			// 3. (Optional:) Save your current layout so you can load it later.
 			// 4. Close the editor.
-			// 5. Manually remove Engine\Saved\Config\Windows\EditorLayout.ini
+			// 5. Manually remove Engine\Saved\Config\Layouts\EditorLayout.ini
 			// 6. Open the Editor, which will auto-regenerate a default EditorLayout.ini that uses your new code below.
 			// 7. "Window" --> "Save Layout" --> "Save Layout As..."
 			//     - Name: Default Editor Layout
@@ -527,8 +532,8 @@ void FMainFrameModule::CreateDefaultMainFrameAuxiliary(const bool bStartImmersiv
 			//     9.2. Engine/Config/BaseEditorLayout.ini
 			//     9.3. Etc
 			// 10. Push the new "DefaultLayout.ini" together with your new code.
-			// 11. Also update these instructions if you change the version number (e.g., from "UnrealEd_Layout_v1.4" to "UnrealEd_Layout_v1.5").
-			const FName LayoutName = TEXT("UnrealEd_Layout_v1.5");
+			// 11. Also update these instructions if you change the version number (e.g., from "UnrealEd_Layout_v1.5" to "UnrealEd_Layout_v1.6").
+			const FName LayoutName = TEXT("UnrealEd_Layout_v1.6");
 			const TSharedRef<FTabManager::FLayout> DefaultLayout =
 				// We persist the positioning of the level editor and the content browser.
 				// The asset editors currently do not get saved.
