Implementing a temporary workaround to patch the potentialy corrupted UnrealEd_Layout_v1.5 has the auto restore to default mecanism is not working in older version of unreal engine 5.

Removed the old config reset and message about the debug file . it hasn't been accurate or doing anything since 2021

#Synced-CL 50202233
#Synced-CL 50202233

diff --git a/Engine/Source/Editor/MainFrame/Private/MainFrameModule.cpp b/Engine/Source/Editor/MainFrame/Private/MainFrameModule.cpp
--- a/Engine/Source/Editor/MainFrame/Private/MainFrameModule.cpp
+++ b/Engine/Source/Editor/MainFrame/Private/MainFrameModule.cpp
@@ -604,7 +604,75 @@ void FMainFrameModule::CreateDefaultMainFrameAuxiliary(const bool bStartImmersiv
 		
 			MainFrameContent = FGlobalTabmanager::Get()->RestoreFrom(LoadedLayout, RootWindow, WindowConfig.bEmbedTitleAreaContent, OutputCanBeNullptr);
 			// MainFrameContent will only be nullptr if its main area contains invalid tabs (probably some layout bug). If so, reset layout to avoid potential crashes
-		
+
+
+			/**
+			 * Workaround a compatibility issue that was introduced in the version 1.5
+			 * The engine is supposed to be able to reset the layout for that type of accidental issue but that mechanism has been dysfunctional since 2021 and was fixed in 5.8
+			 * This fix will be here for an engine a release to help auto fix corrupt config that were created during the window were the issue existed to avoid the crash in the previous release of the engine caused by this issue.
+			 */
+			{
+				static bool bDidPathIssueInPreviousVersion = false;
+				if (!bDidPathIssueInPreviousVersion)
+				{
+					bDidPathIssueInPreviousVersion = true;
+					RemovedOlderLayoutVersions.Reset();
+
+					const FName ProblematicLayoutName = TEXT("UnrealEd_Layout_v1.5");
+					const TSharedRef<FTabManager::FLayout> Layout =
+						// We persist the positioning of the level editor and the content browser.
+						// The asset editors currently do not get saved.
+						FTabManager::NewLayout(ProblematicLayoutName)
+						->AddArea
+						(
+							// level editor window
+							FTabManager::NewPrimaryArea()
+							->Split
+							(
+								FTabManager::NewStack()
+								->SetSizeCoefficient(2.0f)
+								->AddTab("LevelEditor", ETabState::OpenedTab)
+								->AddTab("DockedToolkit", ETabState::ClosedTab)
+							)
+						)
+						->AddArea
+						(
+							// toolkits window
+							FTabManager::NewArea(WindowSize)
+							->SetOrientation(Orient_Vertical)
+							->Split
+							(
+								FTabManager::NewStack()
+								->SetSizeCoefficient(1.0f)
+								->AddTab("StandaloneToolkit", ETabState::ClosedTab)
+							)
+						)
+						->AddArea
+						(
+							// settings window
+							FTabManager::NewArea(WindowSize)
+							->Split
+							(
+								FTabManager::NewStack()
+								->SetSizeCoefficient(1.0f)
+								->AddTab("EditorSettings", ETabState::ClosedTab)
+								->AddTab("ProjectSettings", ETabState::ClosedTab)
+								->AddTab("PluginsEditor", ETabState::ClosedTab)
+							)
+						);
+
+
+					const TSharedRef<FTabManager::FLayout> LayoutToFix = FLayoutSaveRestore::LoadFromConfig(GEditorLayoutIni, Layout, OutputCanBeNullptr, RemovedOlderLayoutVersions);
+
+					const TSharedPtr<FTabManager::FArea> PrimaryArea = LayoutToFix->GetPrimaryArea().Pin();
+					if (PrimaryArea && PrimaryArea->WindowPlacement != FTabManager::FArea::EWindowPlacementInternal::Placement_NoWindow)
+					{
+						PrimaryArea->WindowPlacement = FTabManager::FArea::EWindowPlacementInternal::Placement_NoWindow;
+						FLayoutSaveRestore::SaveToConfig(GEditorLayoutIni, LayoutToFix);
+					}
+				}
+			} 
+
 			if (!MainFrameContent.IsValid())
 			{
 				// This code will recurse to load the default layout. If that also fail we will get into an infinite loop
@@ -617,15 +685,8 @@ void FMainFrameModule::CreateDefaultMainFrameAuxiliary(const bool bStartImmersiv
 				FGlobalTabmanager::Get()->CloseAllAreas();
 				FSlateApplication::Get().CloseAllWindowsImmediately();
 
-				// Remove and reload file
-				GConfig->UnloadFile(GEditorLayoutIni); // We must re-read it to avoid the Editor to use a previously cached name and description
-				const FString FaultyEditorLayoutPath = GEditorLayoutIni + TEXT("_faulty.ini");
-				FPlatformFileManager::Get().GetPlatformFile().MoveFile(*FaultyEditorLayoutPath, *GEditorLayoutIni);
-				GConfig->LoadFile(GEditorLayoutIni);
 				// Warn user/developer
-				const FString WarningMessage = FString::Format(TEXT("UnrealEd layout could not be loaded from the config file {0}, reseting this config file to the default one."), { *GEditorLayoutIni });
-				UE_LOG(LogMainFrame, Warning, TEXT("%s"), *WarningMessage);
-				ensureMsgf(false, TEXT("%s Some additional testing of that layout file should be done. Saved as %s."), *WarningMessage, *FaultyEditorLayoutPath);
+				UE_LOG(LogMainFrame, Warning, TEXT("UnrealEd layout could not be loaded from the config files, we will use the default provided by the code"));
 				// Reload default main frame
 				CreateDefaultMainFrame(bStartImmersive, bStartPIE);
 				return;
diff --git a/Engine/Source/Runtime/Slate/Public/Framework/Docking/TabManager.h b/Engine/Source/Runtime/Slate/Public/Framework/Docking/TabManager.h
--- a/Engine/Source/Runtime/Slate/Public/Framework/Docking/TabManager.h
+++ b/Engine/Source/Runtime/Slate/Public/Framework/Docking/TabManager.h
@@ -691,7 +691,10 @@ class FTabManager : public TSharedFromThis<FTabManager>
 
 		class FArea : public FSplitter
 		{
-				friend class FTabManager;
+			friend class FTabManager;
+
+			// Used for a workaround for a config compatibility issue during the development of 5.8, it will be removed in 5.9.
+			friend class FMainFrameModule;
 		
 			protected:
 				/** An enum which specifies how content appears within a window. */
