Editor: Orphaned tab spawners will no longer be duplicated in the Window menu.

#Synced-CL 50388552

diff --git a/Engine/Source/Editor/MainFrame/Private/Menus/MainMenu.cpp b/Engine/Source/Editor/MainFrame/Private/Menus/MainMenu.cpp
--- a/Engine/Source/Editor/MainFrame/Private/Menus/MainMenu.cpp
+++ b/Engine/Source/Editor/MainFrame/Private/Menus/MainMenu.cpp
@@ -282,11 +282,17 @@ void FMainMenu::RegisterWindowMenu()
 				if(TabManager != FGlobalTabmanager::Get())
 				{
 					// Local editor tabs
-					UE::MainMenu::Private::PopulateTabSpawnerToolMenu(TabManager, InMenu, TabManager->GetLocalWorkspaceMenuRoot(), true);
+					{
+						constexpr bool bDisplayOrphanedTabs = false;
+						UE::MainMenu::Private::PopulateTabSpawnerToolMenu(TabManager, InMenu, TabManager->GetLocalWorkspaceMenuRoot(), bDisplayOrphanedTabs);
+					}
 
 					// General tabs
-					const IWorkspaceMenuStructure& MenuStructure = WorkspaceMenu::GetMenuStructure();
-					UE::MainMenu::Private::PopulateTabSpawnerToolMenu(TabManager, InMenu, MenuStructure.GetStructureRoot(), true);
+					{
+						constexpr bool bDisplayOrphanedTabs = true;
+						const IWorkspaceMenuStructure& MenuStructure = WorkspaceMenu::GetMenuStructure();
+						UE::MainMenu::Private::PopulateTabSpawnerToolMenu(TabManager, InMenu, MenuStructure.GetStructureRoot(), bDisplayOrphanedTabs);
+					}
 				}
 			}
 		}
