[HomeScreen] Remove regex check on links trying to open on an external browser

#Synced-CL 50450741

diff --git a/Engine/Source/Editor/MainFrame/Private/HomeScreen/HomeScreenWeb.cpp b/Engine/Source/Editor/MainFrame/Private/HomeScreen/HomeScreenWeb.cpp
--- a/Engine/Source/Editor/MainFrame/Private/HomeScreen/HomeScreenWeb.cpp
+++ b/Engine/Source/Editor/MainFrame/Private/HomeScreen/HomeScreenWeb.cpp
@@ -20,18 +20,17 @@ void UHomeScreenWeb::OpenGettingStartedProject()
 
 void UHomeScreenWeb::OpenWebPage(const FString& InURL) const
 {
-	const FText ErrorText = FText::Format(LOCTEXT("LaunchingURLNotAllowed", "Given URL is not registered as an allowed one:\n {0}"), FText::FromString(InURL));
+	const FString ValidURLRegex = TEXT("^https?://");
+	const FRegexPattern URLPattern(ValidURLRegex, ERegexPatternFlags::CaseInsensitive);
+	FRegexMatcher URLMatcher(URLPattern, InURL);
 
-	const FString EpicOriginRegex = TEXT("^https?://(www|www2|localhost|dev|docs|forums|company|legal)?(\\.staging)?\\.?(unrealengine|epicgames|twinmotion|metahuman|realityscan|kidswebservices|fab|quixel)(\\.com|-(ci|gamedev)\\.ol\\.epicgames\\.net)(:\\d{1,5})?(/[a-zA-Z0-9._~:/?#@!$&'()*+,;=%-]*/?)*$");
-	const FRegexPattern URLPattern(EpicOriginRegex, ERegexPatternFlags::CaseInsensitive);
-	FRegexMatcher Matcher(URLPattern, InURL);
-
-	if (Matcher.FindNext())
+	if (URLMatcher.FindNext())
 	{
 		FPlatformProcess::LaunchURL(*InURL, nullptr, nullptr);
 	}
 	else
 	{
+		const FText ErrorText = FText::Format(LOCTEXT("LaunchingURLNotValid", "URL is not valid:\n {0}"), FText::FromString(InURL));
 		FMessageDialog::Open(EAppMsgType::Ok, ErrorText);
 	}
 }
