[HomeScreen] Fix crash regarding ReconstructDefaultMainFrameContent

#Synced-CL 50450741

diff --git a/Engine/Source/Editor/MainFrame/Private/MainFrameModule.cpp b/Engine/Source/Editor/MainFrame/Private/MainFrameModule.cpp
--- a/Engine/Source/Editor/MainFrame/Private/MainFrameModule.cpp
+++ b/Engine/Source/Editor/MainFrame/Private/MainFrameModule.cpp
@@ -341,30 +341,25 @@ void FMainFrameModule::RecreateDefaultMainFrame(const bool bStartImmersive, cons
 
 void FMainFrameModule::ReconstructDefaultMainFrameContent(const bool bStartImmersive, const bool bStartPIE)
 {
+	// bRecreatingDefaultMainFrame make sure that the DestroyWindow called on the RootWindow will not be executed
+	check(!bRecreatingDefaultMainFrame);
+	TGuardValue<bool> GuardRecreatingDefaultMainFrame(bRecreatingDefaultMainFrame, true);
+
 	// Clean previous default main frame
 	if (IsWindowInitialized())
 	{
 		// Clean FGlobalTabmanager
 		FGlobalTabmanager::Get()->CloseAllAreas();
 
-		// ::RequestDestroyWindow will remove the window from the TArray so we
-		// should iterate over it backwards to make sure that the WindowIndex is still correct for the TopLevelWindows.
-		TSharedRef<SWindow> RootWindow = GetParentWindow().ToSharedRef();
-
-		// Iterate every TopLevelWindows and destroy them if they are not the MainFrame ParentWindow
+		// Iterate every TopLevelWindows and destroy them unless the window is not regular or is the MainFrame window
 		TArray<TSharedRef<SWindow>> TopLevelWindows = FSlateApplication::Get().GetTopLevelWindows();
-		for (int32 TopLevelWindowIndex = TopLevelWindows.Num() -1; TopLevelWindowIndex >= 0; --TopLevelWindowIndex)
-		{
-			if (TopLevelWindows[TopLevelWindowIndex] == RootWindow)
-			{
-				continue;
-			}
+		DestroyRegularAndNonRootWindows(TopLevelWindows);
 
-			FSlateApplication::Get().RequestDestroyWindow(TopLevelWindows[TopLevelWindowIndex]);
-		}
+		// Iterate every visible windows and destroy them unless the window is not regular or is the MainFrame window
+		TArray<TSharedRef<SWindow>> OutVisibleWindows;
+		FSlateApplication::Get().GetAllVisibleWindowsOrdered(OutVisibleWindows);
+		DestroyRegularAndNonRootWindows(OutVisibleWindows);
 
-		FSlateApplication::Get().DestroyWindowsImmediately();
-		
 		// (Re-)populate the default main frame
 		CreateDefaultMainFrameAuxiliary(bStartImmersive, bStartPIE, false, true);
 	}
@@ -425,7 +420,9 @@ void FMainFrameModule::CreateDefaultMainFrameAuxiliary(const bool bStartImmersiv
 
 		const FText CloseButtonToolTipText = LOCTEXT("MainFrameModuleCloseButtonToolTip", "Close Unreal Editor");
 
-		if (!bIsJustCreatingContent)
+		// If the parent window is not valid it will be initialized below
+		TSharedPtr<SWindow> RootWindowPtr = GetParentWindow();
+		if (!RootWindowPtr.IsValid())
 		{
 			TSharedRef<SWindow> RootWindow = SNew(SWindow)
 				.AutoCenter(WindowConfig.CenterRules)
@@ -455,11 +452,10 @@ void FMainFrameModule::CreateDefaultMainFrameAuxiliary(const bool bStartImmersiv
 			FGlobalTabmanager::Get()->SetAllowWindowMenuBar(true);
 
 			FSlateNotificationManager::Get().SetRootWindow(RootWindow);
+
+			RootWindowPtr = RootWindow;
 		}
 
-		// If we are forcing the ContentCreation to be able to reach here the GetParentWindow should be valid
-		TSharedPtr<SWindow> RootWindowPtr = bIsJustCreatingContent ? GetParentWindow() : FGlobalTabmanager::Get()->GetRootWindow();
-		
 		// The window has to be valid
 		check(RootWindowPtr.IsValid())
 
@@ -1393,6 +1389,39 @@ void FMainFrameModule::HandleCodeAccessorOpenFileFailed(const FString& Filename)
 	FSlateNotificationManager::Get().QueueNotification(Info);
 }
 
+void FMainFrameModule::DestroyRegularAndNonRootWindows(TArray<TSharedRef<SWindow>> InWindowsToDestroy) const
+{
+	// If the root window was destroyed skip this
+	if (IsWindowInitialized())
+	{
+		// Get the root window
+		TSharedRef<SWindow> RootWindow = GetParentWindow().ToSharedRef();
+		
+		// ::RequestDestroyWindow will remove the window from the TArray so we
+		// should iterate over it backwards to make sure that the WindowIndex is still correct.
+		for (int32 WindowIndex = InWindowsToDestroy.Num() -1; WindowIndex >= 0; --WindowIndex)
+		{
+			// Skip the root window
+			if (InWindowsToDestroy[WindowIndex] == RootWindow)
+			{
+				continue;
+			}
+
+			// If it is not a regular window skip it since it may be a pop-up window containing a progress bar for example
+			const bool bIsRegularWindow = InWindowsToDestroy[WindowIndex]->IsRegularWindow();
+			if (!bIsRegularWindow)
+			{
+				continue;
+			}
+
+			InWindowsToDestroy[WindowIndex]->RequestDestroyWindow();
+		}
+		
+		// Destroy the windows immediately from the WindowDestroyQueue
+		FSlateApplication::Get().DestroyWindowsImmediately();
+	}
+}
+
 void FMainFrameModule::RegisterHomeScreenFileMenu()
 {
 	UToolMenus* ToolMenus = UToolMenus::Get();
diff --git a/Engine/Source/Editor/MainFrame/Private/MainFrameModule.h b/Engine/Source/Editor/MainFrame/Private/MainFrameModule.h
--- a/Engine/Source/Editor/MainFrame/Private/MainFrameModule.h
+++ b/Engine/Source/Editor/MainFrame/Private/MainFrameModule.h
@@ -296,6 +296,9 @@ private:
 		DelayedShowMainFrameDelegate.Unbind();
 	}
 
+	// Destroy the given array of windows if they are regular windows and not the root one
+	void DestroyRegularAndNonRootWindows(TArray<TSharedRef<SWindow>> InWindowsToDestroy) const;
+
 	// Register HomeScreen file window menu
 	static void RegisterHomeScreenFileMenu();
 
