When closing the editor save the layout  before asking  checking if we can close the editor as it can modify what the tabmanager would save on exit.

Because the actual windows destruction is delayed by one frame we can capture all the window that where open on the frame where the main frame got closed.
Before this change the CanCloseEditor would sometime display a modal window which would let the other windows fully destroy themselves before we save the layout.

#Synced-CL 50736586

diff --git a/Engine/Source/Editor/MainFrame/Private/Frame/MainFrameHandler.cpp b/Engine/Source/Editor/MainFrame/Private/Frame/MainFrameHandler.cpp
--- a/Engine/Source/Editor/MainFrame/Private/Frame/MainFrameHandler.cpp
+++ b/Engine/Source/Editor/MainFrame/Private/Frame/MainFrameHandler.cpp
@@ -134,14 +134,6 @@ void FMainFrameHandler::ShutDownEditor()
 		RootWindowLocation.SaveToIni();
 	}
 
-	// Save the visual state of the editor before we even
-	// ask whether we can shut down.
-	TSharedRef<FGlobalTabmanager> GlobalTabManager = FGlobalTabmanager::Get();
-	if (GlobalTabManager->CanSavePersistentLayouts())
-	{
-		GlobalTabManager->SaveAllVisualState();
-	}
-
 	// Clear the callback for destructionfrom the main tab; otherwise it will re-enter this shutdown function.
 	if (MainTabPtr.IsValid())
 	{
diff --git a/Engine/Source/Editor/MainFrame/Private/Frame/MainFrameHandler.h b/Engine/Source/Editor/MainFrame/Private/Frame/MainFrameHandler.h
--- a/Engine/Source/Editor/MainFrame/Private/Frame/MainFrameHandler.h
+++ b/Engine/Source/Editor/MainFrame/Private/Frame/MainFrameHandler.h
@@ -192,6 +192,14 @@ public:
 
 	void CloseRootWindowOverride( const TSharedRef<SWindow>& WindowBeingClosed )
 	{
+		// Save the visual state of the editor before we even
+		// ask whether we can shut down.
+		TSharedRef<FGlobalTabmanager> GlobalTabManager = FGlobalTabmanager::Get();
+		if (GlobalTabManager->CanSavePersistentLayouts())
+		{
+			GlobalTabManager->SaveAllVisualState();
+		}
+
 		if (CanCloseEditor())
 		{
 			ShutDownEditor();
