HomeScreen: Delay-load the WebBrowser until the screen becomes actually visible

In a majority of cases, the home tab will never be opened in an editor session with a loaded project, so we're creating the WebBrowser for nothing. Slate supports lazy-instantiation of Widgets via SBox::SetContent() so we can create the HomeScreen contents with an empty SBox placeholder where the webbrowser would go, then do the actual instantiation when we receive the first Tick

#rnx

#Synced-CL 51424673

diff --git a/Engine/Source/Editor/MainFrame/Private/HomeScreen/SHomeScreen.cpp b/Engine/Source/Editor/MainFrame/Private/HomeScreen/SHomeScreen.cpp
--- a/Engine/Source/Editor/MainFrame/Private/HomeScreen/SHomeScreen.cpp
+++ b/Engine/Source/Editor/MainFrame/Private/HomeScreen/SHomeScreen.cpp
@@ -16,6 +16,7 @@
 #include "HomeScreenWeb.h"
 #include "HttpModule.h"
 #include "HttpRetrySystem.h"
+#include "IWebBrowserSingleton.h"
 #include "IWebBrowserWindow.h"
 #include "MainFrameModule.h"
 #include "ProjectDescriptor.h"
@@ -26,6 +27,7 @@
 #include "TimerManager.h"
 #include "ToolMenus.h"
 #include "UnrealEdGlobals.h"
+#include "WebBrowserModule.h"
 #include "Internationalization/Culture.h"
 #include "Widgets/Docking/SDockTab.h"
 #include "Widgets/Images/SImage.h"
@@ -68,10 +70,8 @@ SHomeScreen::~SHomeScreen()
 	RegisterOnDestructedAnalyticEvent();
 }
 
-void SHomeScreen::Construct(const FArguments& InArgs, const TSharedPtr<SWebBrowser> InWebBrowser, const TWeakObjectPtr<UHomeScreenWeb> InWebObject)
+void SHomeScreen::Construct(const FArguments& InArgs)
 {
-	WebBrowser = InWebBrowser;
-	WebObject = InWebObject;
 	SectionLinks = InArgs._SectionLinks;
 	OnGetStartupSettingDelegate = InArgs._OnGetStartupSetting;
 
@@ -92,13 +92,6 @@ void SHomeScreen::Construct(const FArguments& InArgs, const TSharedPtr<SWebBrows
 		AutoLoadProjectComboBoxSelection = LoadAtStartupSetting.GetOnLoadAtStartup() ? EAutoLoadProject::LastProject : EAutoLoadProject::HomeScreen;
 	}
 
-	if (WebObject.IsValid())
-	{
-		WebObject->OnNavigationChanged().AddSP(this, &SHomeScreen::OnNavigateToSection);
-		WebObject->OnTutorialProjectRequested().AddSP(this, &SHomeScreen::OnOpenGettingStartedProject);
-		WebObject->OnOpenLink().AddSP(this, &SHomeScreen::OnWebLinkClicked);
-	}
-
 	bHasProjectsOnMachine = InArgs._AlwaysShowRecentList || HasAlreadyLatestEngineProject();
 	MainHomeSelection = bHasProjectsOnMachine ? EMainSectionMenu::Home : EMainSectionMenu::GettingStarted;
 
@@ -130,6 +123,9 @@ void SHomeScreen::Construct(const FArguments& InArgs, const TSharedPtr<SWebBrows
 	CreateSocialMediaButtons(TwitchButton, TEXT("https://www.twitch.tv/unrealengine"), FAppStyle::GetBrush("SocialMedia.Twitch"));
 	CreateSocialMediaButtons(InstagramButton, TEXT("https://www.instagram.com/unrealengine/"), FAppStyle::GetBrush("SocialMedia.Instagram"));
 
+	// WebBrowser (delay-loaded)
+	SAssignNew(WebBrowserContainer, SBox);
+
 	ChildSlot
 	[
 		SNew(SBorder)
@@ -543,13 +539,13 @@ void SHomeScreen::Construct(const FArguments& InArgs, const TSharedPtr<SWebBrows
 					+ SVerticalBox::Slot()
 					.FillHeight(1.f)
 					[
-						WebBrowser.ToSharedRef()
+						WebBrowserContainer.ToSharedRef()
 					]
 				]
 
 				+ SWidgetSwitcher::Slot()
 				[
-					WebBrowser.ToSharedRef()
+					WebBrowserContainer.ToSharedRef()
 				]
 			]
 		]
@@ -567,10 +563,13 @@ void SHomeScreen::Construct(const FArguments& InArgs, const TSharedPtr<SWebBrows
 				OutButton.ToSharedRef()
 			];
 	}
+}
 
-	// Skip the analytics when selecting the first section during the construction of the widget
-	constexpr bool bSkipAnalytics = true;
-	OnMainHomeSectionChanged(ECheckBoxState::Checked, MainHomeSelection, bSkipAnalytics);
+void SHomeScreen::Tick(const FGeometry& InAllottedGeometry, const double InCurrentTime, const float InDeltaTime)
+{
+	CreateWebBrowserWidgetIfNeeded();
+
+	SCompoundWidget::Tick(InAllottedGeometry, InCurrentTime, InDeltaTime);
 }
 
 void SHomeScreen::OnTabForegrounded(TSharedPtr<SDockTab> NewForegroundTab, TSharedPtr<SDockTab> BackgroundTab)
@@ -830,6 +829,94 @@ void SHomeScreen::CreateSocialMediaButtons(TSharedPtr<SButton>& OutButton, FStri
 			.ColorAndOpacity(this, &SHomeScreen::GetResourceAndSocialMediaButtonColor, OutButton));
 }
 
+void SHomeScreen::CreateWebBrowserWidgetIfNeeded()
+{
+	// WebBrowser already created or failed to create, early exit
+	// Note that we check WebObject instead of WebBrowser to avoid showing error messages repeatedly in case the WebBrowserModule initialization or WebBrowser creation fails
+	if (WebObject.IsValid())
+	{
+		return;
+	}
+
+	WebObject = TStrongObjectPtr(NewObject<UHomeScreenWeb>());
+	WebObject->OnNavigationChanged().AddSP(this, &SHomeScreen::OnNavigateToSection);
+	WebObject->OnTutorialProjectRequested().AddSP(this, &SHomeScreen::OnOpenGettingStartedProject);
+	WebObject->OnOpenLink().AddSP(this, &SHomeScreen::OnWebLinkClicked);
+
+	if (!IWebBrowserModule::Get().IsWebModuleAvailable())
+	{
+		WebBrowserContainer->SetContent(
+			SNew(STextBlock)
+			.Text(LOCTEXT("WebBrowserModuleNotAvailable", "Failed to load the web browser module, can't display web content!"))
+			.Font(FAppStyle::GetFontStyle("HomeScreen.Font.DefaultSize"))
+			.Justification(ETextJustify::Center)
+		);
+		return;
+	}
+
+	IWebBrowserSingleton* WebBrowserSingleton = IWebBrowserModule::Get().GetSingleton();
+
+	FCreateBrowserWindowSettings WindowSettings;
+	WindowSettings.bUseTransparency = true;
+	TSharedPtr<IWebBrowserWindow> WebBrowserWindow = WebBrowserSingleton->CreateBrowserWindow(WindowSettings);
+	if (!WebBrowserWindow.IsValid())
+	{
+		WebBrowserContainer->SetContent(
+			SNew(STextBlock)
+			.Text(LOCTEXT("WebBrowserWindowCreationIssue", "Failed to create the web browser window, can't display web content!"))
+			.Font(FAppStyle::GetFontStyle("HomeScreen.Font.DefaultSize"))
+			.Justification(ETextJustify::Center)
+		);
+		return;
+	}
+
+#if !UE_BUILD_SHIPPING
+	WebBrowserSingleton->SetDevToolsShortcutEnabled(true);
+	WebBrowserWindow->OnCreateWindow().BindSPLambda(this, [](const TWeakPtr<IWebBrowserWindow>& InNewBrowserWindowWeak, const TWeakPtr<IWebBrowserPopupFeatures>& PopupFeatures)
+		{
+			if (const TSharedPtr<IWebBrowserWindow> NewBrowserWindow = InNewBrowserWindowWeak.Pin())
+			{
+				// Initialize a dialog
+				auto DialogMainWindow = SNew(SWindow)
+					.Title(FText::FromString(TEXT("Chrome Debugging Tools")))
+					.ClientSize(FVector2D(700, 700))
+					.SupportsMaximize(true)
+					.SupportsMinimize(true)
+					[
+						SNew(SVerticalBox)
+
+						+ SVerticalBox::Slot()
+						.HAlign(HAlign_Fill)
+						.VAlign(VAlign_Fill)
+						[
+							SNew(SWebBrowser, NewBrowserWindow.ToSharedRef())
+						]
+					];
+				FSlateApplication::Get().AddWindow(DialogMainWindow);
+				return true;
+			}
+			return false;
+		});
+#endif
+
+	SAssignNew(WebBrowser, SWebBrowser, WebBrowserWindow)
+		.OnLoadStarted(FSimpleDelegate::CreateSPLambda(this, [this]() { WebBrowser->ShowThrobber(); }))
+		.OnLoadError(FSimpleDelegate::CreateSPLambda(this, [this]() { WebBrowser->HideThrobber(); }))
+		.OnLoadCompleted(FSimpleDelegate::CreateSPLambda(this, [this]() { WebBrowser->HideThrobber(); }))
+		.ShowControls(false)
+		.SupportsTransparency(true)
+		.ShowAddressBar(false);
+
+	WebBrowserContainer->SetContent(WebBrowser.ToSharedRef());
+
+	constexpr bool bIsPermanent = true;
+	WebBrowser->BindUObject(TEXT("uebridge"), WebObject.Get(), bIsPermanent);
+
+	// Skip the analytics when selecting the first section during the construction of the widget
+	constexpr bool bSkipAnalytics = true;
+	OnMainHomeSectionChanged(ECheckBoxState::Checked, MainHomeSelection, bSkipAnalytics);
+}
+
 TSharedRef<SWidget> SHomeScreen::CreateComboButtonMenuContentWidget()
 {
 	UToolMenus* ToolMenus = UToolMenus::Get();
diff --git a/Engine/Source/Editor/MainFrame/Private/HomeScreen/SHomeScreen.h b/Engine/Source/Editor/MainFrame/Private/HomeScreen/SHomeScreen.h
--- a/Engine/Source/Editor/MainFrame/Private/HomeScreen/SHomeScreen.h
+++ b/Engine/Source/Editor/MainFrame/Private/HomeScreen/SHomeScreen.h
@@ -13,12 +13,14 @@
 #include "Widgets/SNullWidget.h"
 #include "Widgets/SWidget.h"
 
+class SBox;
 class SButton;
 class SCheckBox;
 class SComboButton;
 class SDockTab;
 class SVerticalBox;
 class SWebBrowser;
+class UHomeScreenWeb;
 template <typename ItemType> class STileView;
 
 class SHomeScreen : public SCompoundWidget
@@ -43,11 +45,15 @@ public:
 	virtual ~SHomeScreen() override;
 
 	/** Constructs the HomeScreen widget */
-	void Construct(const FArguments& InArgs, const TSharedPtr<SWebBrowser> InWebBrowser, const TWeakObjectPtr<UHomeScreenWeb> InWebObject);
+	void Construct(const FArguments& InArgs);
 
 	/** Used when the HomeScreen tab becomes visible to start the internet connection check or to stop it when it goes in the background */
 	void OnTabForegrounded(TSharedPtr<SDockTab> NewForegroundTab, TSharedPtr<SDockTab> BackgroundTab);
 
+	//~ Begin SWidget
+	virtual void Tick(const FGeometry& InAllottedGeometry, const double InCurrentTime, const float InDeltaTime) override;
+	//~ End SWidget
+
 private:
 	/** Opens the Create Project dialog */
 	FReply OnCreateProjectDialog(bool bInAllowProjectOpening, bool bInAllowProjectCreation);
@@ -94,6 +100,9 @@ private:
 	/** Creates social media buttons */
 	void CreateSocialMediaButtons(TSharedPtr<SButton>& OutButton, FString InLink, const FSlateBrush* InImage);
 
+	/** Creates the webbrowser widget if not initialized yet */
+	void CreateWebBrowserWidgetIfNeeded();
+
 	/** Creates the content widget for the combo button menu */
 	TSharedRef<SWidget> CreateComboButtonMenuContentWidget();
 
@@ -190,11 +199,14 @@ private:
 	/** Selection for autoload project combo box */
 	EAutoLoadProject AutoLoadProjectComboBoxSelection;
 
-	/** WebBrowser widget passed down from the MainFrameModule */
+	/** Web browser lazy-loading container */
+	TSharedPtr<SBox> WebBrowserContainer;
+
+	/** Web browser widget instance */
 	TSharedPtr<SWebBrowser> WebBrowser;
-	
-	/** WebBrowser object passed down from the MainFrameModule */
-	TWeakObjectPtr<UHomeScreenWeb> WebObject;
+
+	/** Bridge object for browser navigation */
+	TStrongObjectPtr<UHomeScreenWeb> WebObject;
 
 	/** HomeScreen start up setting getter */
 	FOnGetStartupSetting OnGetStartupSettingDelegate;
diff --git a/Engine/Source/Editor/MainFrame/Private/MainFrameModule.cpp b/Engine/Source/Editor/MainFrame/Private/MainFrameModule.cpp
--- a/Engine/Source/Editor/MainFrame/Private/MainFrameModule.cpp
+++ b/Engine/Source/Editor/MainFrame/Private/MainFrameModule.cpp
@@ -40,14 +40,10 @@
 #include "HomeScreen/HomeScreenWeb.h"
 #include "HomeScreen/SHomeScreen.h" 
 #include "Internationalization/Culture.h"
-#include "IWebBrowserSingleton.h"
-#include "IWebBrowserWindow.h"
 #include "Menus/SettingsMenu.h"
 #include "Settings/EditorStyleSettings.h"
 #include "Settings/HomeScreenSettings.h"
-#include "SWebBrowser.h"
 #include "ToolMenus.h"
-#include "WebBrowserModule.h"
 #include "Widgets/Docking/SDockTab.h"
 
 DEFINE_LOG_CATEGORY(LogMainFrame);
@@ -1031,9 +1027,6 @@ void FMainFrameModule::ShutdownModule( )
 	}
 
 	// Clean up static object/widget used for the HomeScreen Tab
-	WebObject.Reset();
-	WebBrowserWindow.Reset();
-	WebBrowser.Reset();
 	HomeScreenTabManager.Reset();
 
 	// Unregister HomeScreen tab on shutdown
@@ -1290,94 +1283,7 @@ TSharedRef<SDockTab> FMainFrameModule::SpawnHomeScreenWindow(const FSpawnTabArgs
 
 TSharedRef<SWidget> FMainFrameModule::CreateHomeScreenWidget(const TSharedPtr<SDockTab>& InParentDockTab, FHomeScreenWidgetSettings InSettings)
 {
-	if (!WebObject.IsValid())
-	{
-		WebObject = TStrongObjectPtr(NewObject<UHomeScreenWeb>());
-	}
-
-	if (!WebBrowserWindow.IsValid())
-	{
-		IWebBrowserModule& WebBrowserModule = IWebBrowserModule::Get();
-		if (!IWebBrowserModule::IsAvailable() || !WebBrowserModule.IsWebModuleAvailable())
-		{
-			FMessageDialog::Open(EAppMsgType::Ok, LOCTEXT("WebBrowserModuleNotAvailable", "Failed to load the plugin. Please enable Web WebBrowserWindow in the plugin manager to use the Web Panels."));
-			return SNullWidget::NullWidget;
-		}
-
-		IWebBrowserSingleton* WebBrowserSingleton = WebBrowserModule.GetSingleton();
-
-		FCreateBrowserWindowSettings WindowSettings;
-		WindowSettings.InitialURL = TEXT("about:blank");
-		WindowSettings.bUseTransparency = true;
-		WindowSettings.BrowserFrameRate = 60;
-
-		WebBrowserWindow = WebBrowserSingleton->CreateBrowserWindow(WindowSettings);
-		if (!WebBrowserWindow.IsValid())
-		{
-			FMessageDialog::Open(EAppMsgType::Ok, LOCTEXT("WebBrowserWindowCreationIssue", "Failed to create the web browser window."));
-			return SNullWidget::NullWidget;
-		}
-		
-		WebBrowserSingleton->SetDevToolsShortcutEnabled(true);
-		WebBrowserWindow->OnCreateWindow().BindLambda([](const TWeakPtr<IWebBrowserWindow>& InNewBrowserWindowWeak, const TWeakPtr<IWebBrowserPopupFeatures>& PopupFeatures)
-		{
-			if (const TSharedPtr<IWebBrowserWindow> NewBrowserWindow = InNewBrowserWindowWeak.Pin())
-			{
-				// Initialize a dialog
-				auto DialogMainWindow = SNew(SWindow)
-					.Title(FText::FromString(TEXT("Chrome Debugging Tools")))
-					.ClientSize(FVector2D(700, 700))
-					.SupportsMaximize(true)
-					.SupportsMinimize(true)
-					[
-						SNew(SVerticalBox)
-
-						+ SVerticalBox::Slot()
-						.HAlign(HAlign_Fill)
-						.VAlign(VAlign_Fill)
-						[
-							SNew(SWebBrowser, NewBrowserWindow.ToSharedRef())
-						]
-					];
-				FSlateApplication::Get().AddWindow(DialogMainWindow);
-				return true;
-			}
-			return false;
-		});
-	}
-
-	auto ShowThrobber = [this]()
-		{
-			if (WebBrowser.IsValid())
-			{
-				WebBrowser->ShowThrobber();
-			}
-		};
-
-	auto HideThrobber = [this]()
-		{
-			if (WebBrowser.IsValid())
-			{
-				WebBrowser->HideThrobber();
-			}
-		};
-
-	if (!WebBrowser.IsValid())
-	{
-		SAssignNew(WebBrowser, SWebBrowser, WebBrowserWindow)
-		.OnLoadStarted_Lambda(ShowThrobber)
-		.OnLoadError_Lambda(HideThrobber)
-		.OnLoadCompleted_Lambda(HideThrobber)
-		.ShowControls(false)
-		.SupportsTransparency(true)
-		.ShowAddressBar(false);
-
-		constexpr bool bIsPermanent = true;
-		WebBrowser->BindUObject(TEXT("uebridge"), WebObject.Get(), bIsPermanent);
-		WebBrowserWindow->Reload();
-	}
-
-	return SNew(SHomeScreen, WebBrowser, TWeakObjectPtr(WebObject.Get()))
+	return SNew(SHomeScreen)
 		.SectionLinks(InSettings.SectionLinks)
 		.ResourceCommunityEntries(InSettings.ResourcesEntries)
 		.AlwaysShowRecentList(InSettings.bAlwaysShowRecentList)
diff --git a/Engine/Source/Editor/MainFrame/Private/MainFrameModule.h b/Engine/Source/Editor/MainFrame/Private/MainFrameModule.h
--- a/Engine/Source/Editor/MainFrame/Private/MainFrameModule.h
+++ b/Engine/Source/Editor/MainFrame/Private/MainFrameModule.h
@@ -14,10 +14,6 @@
 #include "Interfaces/IEditorMainFrameProvider.h"
 #include "Interfaces/IHomeScreenProvider.h"
 
-class SWebBrowser;
-class IWebBrowserWindow;
-class UHomeScreenWeb;
-
 /** 
  * Utility class, which hooks into the editor's home screen tab and
  * supply it with its own settings if it is the higher priority provider
@@ -372,15 +368,6 @@ private:
 	// Set to true if we are in the middle of recreating the default layout because a previous load attempt failed
 	bool bIsReCreatingDefaultLayoutAsFallback = false;
 
-	// Bridge object for browser navigation
-	TStrongObjectPtr<UHomeScreenWeb> WebObject = nullptr;
-
-	// Web browser window instance
-	TSharedPtr<IWebBrowserWindow> WebBrowserWindow = nullptr;
-
-	// Web browser widget instance
-	TSharedPtr<SWebBrowser> WebBrowser = nullptr;
-
 	// Holds the tab manager that manages the HomeScreen tab.
 	TSharedPtr<FTabManager> HomeScreenTabManager = nullptr;
 
